Copying //Tasks/UE4/Dev-Online-Plugins-v3 to //UE4/Dev-Main (Source: //Tasks/UE4/Dev-Online-Plugins-v3 @ 3047294)

diff --git a/Engine/Plugins/Online/OnlineSubsystemNull/Source/Private/OnlineIdentityNull.h b/Engine/Plugins/Online/OnlineSubsystemNull/Source/Private/OnlineIdentityNull.h
new file mode 100644
--- /dev/null
+++ b/Engine/Plugins/Online/OnlineSubsystemNull/Source/Private/OnlineIdentityNull.h
@@ -0,0 +1,102 @@
+// Copyright 1998-2016 Epic Games, Inc. All Rights Reserved.
+
+#pragma once
+ 
+#include "OnlineIdentityInterface.h"
+
+/**
+ * Info associated with an user account generated by this online service
+ */
+class FUserOnlineAccountNull : 
+	public FUserOnlineAccount
+{
+
+public:
+
+	// FOnlineUser
+	
+	virtual TSharedRef<const FUniqueNetId> GetUserId() const override { return UserIdPtr; }
+	virtual FString GetRealName() const override { return TEXT("DummyRealName"); }
+	virtual FString GetDisplayName(const FString& Platform = FString()) const override  { return TEXT("DummyDisplayName"); }
+	virtual bool GetUserAttribute(const FString& AttrName, FString& OutAttrValue) const override;
+	virtual bool SetUserAttribute(const FString& AttrName, const FString& AttrValue) override;
+
+	// FUserOnlineAccount
+
+	virtual FString GetAccessToken() const override { return TEXT("DummyAuthTicket"); }
+	virtual bool GetAuthAttribute(const FString& AttrName, FString& OutAttrValue) const override;
+
+	// FUserOnlineAccountNull
+
+	FUserOnlineAccountNull(const FString& InUserId=TEXT("")) 
+		: UserIdPtr(new FUniqueNetIdString(InUserId))
+	{ }
+
+	virtual ~FUserOnlineAccountNull()
+	{
+	}
+
+	/** User Id represented as a FUniqueNetId */
+	TSharedRef<const FUniqueNetId> UserIdPtr;
+
+        /** Additional key/value pair data related to auth */
+	TMap<FString, FString> AdditionalAuthData;
+        /** Additional key/value pair data related to user attribution */
+	TMap<FString, FString> UserAttributes;
+};
+
+/**
+ * Null service implementation of the online identity interface
+ */
+class FOnlineIdentityNull : public IOnlineIdentity
+{
+public:
+
+	// IOnlineIdentity
+
+	virtual bool Login(int32 LocalUserNum, const FOnlineAccountCredentials& AccountCredentials) override;
+	virtual bool Logout(int32 LocalUserNum) override;
+	virtual bool AutoLogin(int32 LocalUserNum) override;
+	virtual TSharedPtr<FUserOnlineAccount> GetUserAccount(const FUniqueNetId& UserId) const override;
+	virtual TArray<TSharedPtr<FUserOnlineAccount> > GetAllUserAccounts() const override;
+	virtual TSharedPtr<const FUniqueNetId> GetUniquePlayerId(int32 LocalUserNum) const override;
+	virtual TSharedPtr<const FUniqueNetId> CreateUniquePlayerId(uint8* Bytes, int32 Size) override;
+	virtual TSharedPtr<const FUniqueNetId> CreateUniquePlayerId(const FString& Str) override;
+	virtual ELoginStatus::Type GetLoginStatus(int32 LocalUserNum) const override;
+	virtual ELoginStatus::Type GetLoginStatus(const FUniqueNetId& UserId) const override;
+	virtual FString GetPlayerNickname(int32 LocalUserNum) const override;
+	virtual FString GetPlayerNickname(const FUniqueNetId& UserId) const override;
+	virtual FString GetAuthToken(int32 LocalUserNum) const override;
+	virtual void GetUserPrivilege(const FUniqueNetId& UserId, EUserPrivileges::Type Privilege, const FOnGetUserPrivilegeCompleteDelegate& Delegate) override;
+	virtual FPlatformUserId GetPlatformUserIdFromUniqueNetId(const FUniqueNetId& UniqueNetId) override;
+	virtual FString GetAuthType() const override;
+
+	// FOnlineIdentityNull
+
+	/**
+	 * Constructor
+	 *
+	 * @param InSubsystem online subsystem being used
+	 */
+	FOnlineIdentityNull(class FOnlineSubsystemNull* InSubsystem);
+
+	/**
+	 * Destructor
+	 */
+	virtual ~FOnlineIdentityNull();
+
+private:
+
+	/**
+	 * Should use the initialization constructor instead
+	 */
+	FOnlineIdentityNull();
+
+	/** Ids mapped to locally registered users */
+	TMap<int32, TSharedPtr<const FUniqueNetId>> UserIds;
+
+	/** Ids mapped to locally registered users */
+	TMap<FUniqueNetIdString, TSharedRef<FUserOnlineAccountNull>> UserAccounts;
+};
+
+typedef TSharedPtr<FOnlineIdentityNull, ESPMode::ThreadSafe> FOnlineIdentityNullPtr;
diff --git a/Engine/Plugins/Online/OnlineSubsystemNull/Source/Private/OnlineSubsystemNullPrivatePCH.h b/Engine/Plugins/Online/OnlineSubsystemNull/Source/Private/OnlineSubsystemNullPrivatePCH.h
new file mode 100644
--- /dev/null
+++ b/Engine/Plugins/Online/OnlineSubsystemNull/Source/Private/OnlineSubsystemNullPrivatePCH.h
@@ -0,0 +1,20 @@
+// Copyright 1998-2016 Epic Games, Inc. All Rights Reserved.
+
+#pragma once
+
+#include "Engine.h"
+#include "OnlineSubsystemNullModule.h"
+#include "OnlineSubsystemModule.h"
+#include "OnlineSubsystem.h"
+#include "ModuleManager.h"
+
+#define INVALID_INDEX -1
+
+/** URL Prefix when using Null socket connection */
+#define NULL_URL_PREFIX TEXT("Null.")
+
+/** pre-pended to all NULL logging */
+#undef ONLINE_LOG_PREFIX
+#define ONLINE_LOG_PREFIX TEXT("NULL: ")
+
+
diff --git a/Engine/Plugins/Online/OnlineSubsystemNull/Source/Public/OnlineSubsystemNull.h b/Engine/Plugins/Online/OnlineSubsystemNull/Source/Public/OnlineSubsystemNull.h
new file mode 100644
--- /dev/null
+++ b/Engine/Plugins/Online/OnlineSubsystemNull/Source/Public/OnlineSubsystemNull.h
@@ -0,0 +1,126 @@
+// Copyright 1998-2016 Epic Games, Inc. All Rights Reserved.
+
+#pragma once
+
+#include "OnlineSubsystem.h"
+#include "OnlineSubsystemImpl.h"
+#include "OnlineSubsystemNullPackage.h"
+
+/** Forward declarations of all interface classes */
+typedef TSharedPtr<class FOnlineSessionNull, ESPMode::ThreadSafe> FOnlineSessionNullPtr;
+typedef TSharedPtr<class FOnlineProfileNull, ESPMode::ThreadSafe> FOnlineProfileNullPtr;
+typedef TSharedPtr<class FOnlineFriendsNull, ESPMode::ThreadSafe> FOnlineFriendsNullPtr;
+typedef TSharedPtr<class FOnlineUserCloudNull, ESPMode::ThreadSafe> FOnlineUserCloudNullPtr;
+typedef TSharedPtr<class FOnlineLeaderboardsNull, ESPMode::ThreadSafe> FOnlineLeaderboardsNullPtr;
+typedef TSharedPtr<class FOnlineVoiceImpl, ESPMode::ThreadSafe> FOnlineVoiceImplPtr;
+typedef TSharedPtr<class FOnlineExternalUINull, ESPMode::ThreadSafe> FOnlineExternalUINullPtr;
+typedef TSharedPtr<class FOnlineIdentityNull, ESPMode::ThreadSafe> FOnlineIdentityNullPtr;
+typedef TSharedPtr<class FOnlineAchievementsNull, ESPMode::ThreadSafe> FOnlineAchievementsNullPtr;
+
+/**
+ *	OnlineSubsystemNull - Implementation of the online subsystem for Null services
+ */
+class ONLINESUBSYSTEMNULL_API FOnlineSubsystemNull : 
+	public FOnlineSubsystemImpl
+{
+
+public:
+
+	virtual ~FOnlineSubsystemNull()
+	{
+	}
+
+	// IOnlineSubsystem
+
+	virtual IOnlineSessionPtr GetSessionInterface() const override;
+	virtual IOnlineFriendsPtr GetFriendsInterface() const override;
+	virtual IOnlinePartyPtr GetPartyInterface() const override;
+	virtual IOnlineGroupsPtr GetGroupsInterface() const override;
+	virtual IOnlineSharedCloudPtr GetSharedCloudInterface() const override;
+	virtual IOnlineUserCloudPtr GetUserCloudInterface() const override;
+	virtual IOnlineEntitlementsPtr GetEntitlementsInterface() const override;
+	virtual IOnlineLeaderboardsPtr GetLeaderboardsInterface() const override;
+	virtual IOnlineVoicePtr GetVoiceInterface() const override;
+	virtual IOnlineExternalUIPtr GetExternalUIInterface() const override;	
+	virtual IOnlineTimePtr GetTimeInterface() const override;
+	virtual IOnlineIdentityPtr GetIdentityInterface() const override;
+	virtual IOnlineTitleFilePtr GetTitleFileInterface() const override;
+	virtual IOnlineStorePtr GetStoreInterface() const override;
+	virtual IOnlineStoreV2Ptr GetStoreV2Interface() const override { return nullptr; }
+	virtual IOnlinePurchasePtr GetPurchaseInterface() const override { return nullptr; }
+	virtual IOnlineEventsPtr GetEventsInterface() const override;
+	virtual IOnlineAchievementsPtr GetAchievementsInterface() const override;
+	virtual IOnlineSharingPtr GetSharingInterface() const override;
+	virtual IOnlineUserPtr GetUserInterface() const override;
+	virtual IOnlineMessagePtr GetMessageInterface() const override;
+	virtual IOnlinePresencePtr GetPresenceInterface() const override;
+	virtual IOnlineChatPtr GetChatInterface() const override;
+    virtual IOnlineTurnBasedPtr GetTurnBasedInterface() const override;
+	
+	virtual bool Init() override;
+	virtual bool Shutdown() override;
+	virtual FString GetAppId() const override;
+	virtual bool Exec(class UWorld* InWorld, const TCHAR* Cmd, FOutputDevice& Ar) override;
+
+	// FTickerObjectBase
+	
+	virtual bool Tick(float DeltaTime) override;
+
+	// FOnlineSubsystemNull
+
+	/**
+	 * Is the Null API available for use
+	 * @return true if Null functionality is available, false otherwise
+	 */
+	bool IsEnabled();
+
+PACKAGE_SCOPE:
+
+	/** Only the factory makes instances */
+	FOnlineSubsystemNull(FName InInstanceName) :
+		FOnlineSubsystemImpl(InInstanceName),
+		SessionInterface(nullptr),
+		VoiceInterface(nullptr),
+		LeaderboardsInterface(nullptr),
+		IdentityInterface(nullptr),
+		AchievementsInterface(nullptr),
+		OnlineAsyncTaskThreadRunnable(nullptr),
+		OnlineAsyncTaskThread(nullptr)
+	{}
+
+	FOnlineSubsystemNull() :
+		SessionInterface(nullptr),
+		VoiceInterface(nullptr),
+		LeaderboardsInterface(nullptr),
+		IdentityInterface(nullptr),
+		AchievementsInterface(nullptr),
+		OnlineAsyncTaskThreadRunnable(nullptr),
+		OnlineAsyncTaskThread(nullptr)
+	{}
+
+private:
+
+	/** Interface to the session services */
+	FOnlineSessionNullPtr SessionInterface;
+
+	/** Interface for voice communication */
+	FOnlineVoiceImplPtr VoiceInterface;
+
+	/** Interface to the leaderboard services */
+	FOnlineLeaderboardsNullPtr LeaderboardsInterface;
+
+	/** Interface to the identity registration/auth services */
+	FOnlineIdentityNullPtr IdentityInterface;
+
+	/** Interface for achievements */
+	FOnlineAchievementsNullPtr AchievementsInterface;
+
+	/** Online async task runnable */
+	class FOnlineAsyncTaskManagerNull* OnlineAsyncTaskThreadRunnable;
+
+	/** Online async task thread */
+	class FRunnableThread* OnlineAsyncTaskThread;
+};
+
+typedef TSharedPtr<FOnlineSubsystemNull, ESPMode::ThreadSafe> FOnlineSubsystemNullPtr;
+
diff --git a/Engine/Plugins/Online/OnlineSubsystemSteam/OnlineSubsystemSteam.uplugin b/Engine/Plugins/Online/OnlineSubsystemSteam/OnlineSubsystemSteam.uplugin
new file mode 100644
--- /dev/null
+++ b/Engine/Plugins/Online/OnlineSubsystemSteam/OnlineSubsystemSteam.uplugin
@@ -0,0 +1,28 @@
+{
+	"FileVersion" : 3,
+	
+	"FriendlyName" : "Online Subsystem Steam",
+	"Version" : 1,
+	"VersionName" : "1.0",
+	"EngineVersion" : "4.2.0",
+	"Description" : "Access to Steam platform",
+	"Category" : "Online Platform",
+	"CreatedBy" : "Epic Games, Inc.",
+	"CreatedByURL" : "http://epicgames.com",
+	"EnabledByDefault" : false,
+
+    "Modules": [
+        {
+            "Name": "OnlineSubsystemSteam",
+            "Type": "Runtime",
+            "LoadingPhase": "Default",
+            "WhitelistPlatforms" :
+			[
+                "Win32",
+                "Win64",
+                "Mac",
+                "Linux"
+			]
+        }
+    ]
+}
\ No newline at end of file
diff --git a/Engine/Plugins/Online/OnlineSubsystemSteam/Source/OnlineSubsystemSteam.Build.cs b/Engine/Plugins/Online/OnlineSubsystemSteam/Source/OnlineSubsystemSteam.Build.cs
new file mode 100644
--- /dev/null
+++ b/Engine/Plugins/Online/OnlineSubsystemSteam/Source/OnlineSubsystemSteam.Build.cs
@@ -0,0 +1,38 @@
+// Copyright 1998-2016 Epic Games, Inc. All Rights Reserved.
+
+using UnrealBuildTool;
+using System.IO;
+
+public class OnlineSubsystemSteam : ModuleRules
+{
+	public OnlineSubsystemSteam(TargetInfo Target)
+	{
+		string SteamVersion = "Steamv132";
+		bool bSteamSDKFound = Directory.Exists(UEBuildConfiguration.UEThirdPartySourceDirectory + "Steamworks/" + SteamVersion) == true;
+
+		Definitions.Add("STEAMSDK_FOUND=" + (bSteamSDKFound ? "1" : "0"));
+		Definitions.Add("WITH_STEAMWORKS=" + (bSteamSDKFound ? "1" : "0"));
+
+		Definitions.Add("ONLINESUBSYSTEMSTEAM_PACKAGE=1");
+
+		PublicDependencyModuleNames.AddRange(
+			new string[] { 
+				"OnlineSubsystemUtils",
+			}
+		);
+
+		PrivateDependencyModuleNames.AddRange(
+			new string[] {
+				"Core", 
+				"CoreUObject", 
+				"Engine", 
+				"Sockets", 
+				"Voice",
+				"OnlineSubsystem",
+				"Json"
+			}
+		);
+
+		AddEngineThirdPartyPrivateStaticDependencies(Target, "Steamworks");
+	}
+}
diff --git a/Engine/Plugins/Online/OnlineSubsystemSteam/Source/Private/IPAddressSteam.cpp b/Engine/Plugins/Online/OnlineSubsystemSteam/Source/Private/IPAddressSteam.cpp
new file mode 100644
--- /dev/null
+++ b/Engine/Plugins/Online/OnlineSubsystemSteam/Source/Private/IPAddressSteam.cpp
@@ -0,0 +1,57 @@
+// Copyright 1998-2016 Epic Games, Inc. All Rights Reserved.
+
+#include "OnlineSubsystemSteamPrivatePCH.h"
+#include "IPAddressSteam.h"
+
+/**
+ * Sets the ip address from a string ("steam.STEAMID" or "STEAMID")
+ *
+ * @param InAddr the string containing the new ip address to use
+ * @param bIsValid was the address valid for use
+ */
+void FInternetAddrSteam::SetIp(const TCHAR* InAddr, bool& bIsValid)
+{
+	bIsValid = false;
+	FString InAddrStr(InAddr);
+
+	// Check for the steam. prefix
+	FString SteamIPAddrStr;
+	if (InAddrStr.StartsWith(STEAM_URL_PREFIX))
+	{
+		SteamIPAddrStr = InAddrStr.Mid(ARRAY_COUNT(STEAM_URL_PREFIX) - 1);
+	}
+	else
+	{
+		SteamIPAddrStr = InAddrStr;
+	}
+
+	// Resolve the address/port
+	FString SteamIPStr, SteamChannelStr;
+	if (SteamIPAddrStr.Split(":", &SteamIPStr, &SteamChannelStr, ESearchCase::CaseSensitive, ESearchDir::FromEnd))
+	{
+		const uint64 Id = FCString::Atoi64(*SteamIPStr);
+		if (Id != 0)
+		{
+			SteamId = FUniqueNetIdSteam(Id);
+			const int32 Channel = FCString::Atoi(*SteamChannelStr);
+			if (Channel != 0 || SteamChannelStr == "0")
+			{
+				SteamChannel = Channel;
+				bIsValid = true;
+			}
+		}
+	}
+	else
+	{
+		const uint64 Id = FCString::Atoi64(*SteamIPAddrStr);
+		if (Id != 0)
+		{
+			SteamId = FUniqueNetIdSteam(Id);
+			bIsValid = true;
+		}
+
+		SteamChannel = 0;
+	}
+
+	bIsValid = bIsValid && SteamId.IsValid();
+}
diff --git a/Engine/Plugins/Online/OnlineSubsystemSteam/Source/Private/IPAddressSteam.h b/Engine/Plugins/Online/OnlineSubsystemSteam/Source/Private/IPAddressSteam.h
new file mode 100644
--- /dev/null
+++ b/Engine/Plugins/Online/OnlineSubsystemSteam/Source/Private/IPAddressSteam.h
@@ -0,0 +1,155 @@
+// Copyright 1998-2016 Epic Games, Inc. All Rights Reserved.
+
+#pragma once
+
+#include "Core.h"
+#include "IPAddress.h"
+#include "OnlineSubsystemSteamTypes.h"
+#include "OnlineSubsystemSteamPackage.h"
+
+/**
+ * Represents an internet ip address, using the relatively standard SOCKADDR_IN structure. All data is in network byte order
+ */
+class FInternetAddrSteam : public FInternetAddr
+{
+PACKAGE_SCOPE:
+	/** The Steam id to connect to */
+	FUniqueNetIdSteam SteamId;
+	/** Steam channel to communicate on */
+	int32 SteamChannel;
+
+	/**
+	 * Copy Constructor
+	 */
+	FInternetAddrSteam(const FInternetAddrSteam& Src) :
+		SteamId(Src.SteamId),
+		SteamChannel(Src.SteamChannel)
+	{
+	}
+
+public:
+	/**
+	 * Constructor. Sets address to default state
+	 */
+	FInternetAddrSteam() :
+		SteamId((uint64)0),
+		SteamChannel(0)
+	{
+	}
+	/**
+	 * Constructor. Sets address to default state
+	 */
+	explicit FInternetAddrSteam(const FUniqueNetIdSteam& InSteamId) :
+		SteamId(InSteamId),
+		SteamChannel(0)
+	{
+	}
+
+	/**
+	 * Sets the ip address from a host byte order uint32
+	 *
+	 * @param InAddr the new address to use (must convert to network byte order)
+	 */
+	void SetIp(uint32 InAddr) override
+	{
+		check("Invalid interface for Steam P2P addresses" && 0);
+	}
+
+	/**
+	 * Sets the ip address from a string ("A.B.C.D")
+	 *
+	 * @param InAddr the string containing the new ip address to use
+	 */
+	void SetIp(const TCHAR* InAddr, bool& bIsValid) override;
+
+	/**
+	 * Copies the network byte order ip address to a host byte order dword
+	 *
+	 * @param OutAddr the out param receiving the ip address
+	 */
+	void GetIp(uint32& OutAddr) const override
+	{ 
+		check("Invalid interface for Steam P2P addresses" && 0);
+	}
+
+	/**
+	 * Sets the port number from a host byte order int
+	 *
+	 * @param InPort the new port to use (must convert to network byte order)
+	 */
+	void SetPort(int32 InPort) override
+	{
+		SteamChannel = InPort;
+	}
+
+	/**
+	 * Copies the port number from this address and places it into a host byte order int
+	 *
+	 * @param OutPort the host byte order int that receives the port
+	 */
+	void GetPort(int32& OutPort) const override
+	{
+		OutPort = SteamChannel;
+	}
+
+	/**
+	 * Returns the port number from this address in host byte order
+	 */
+	int32 GetPort() const override
+	{
+		return SteamChannel;
+	}
+
+	/** Sets the address to be any address */
+	void SetAnyAddress() override
+	{
+	}
+
+	/** Sets the address to broadcast */
+	void SetBroadcastAddress() override
+	{
+	}
+
+	/**
+	 * Converts this internet ip address to string form
+	 *
+	 * @param bAppendPort whether to append the port information or not
+	 */
+	FString ToString(bool bAppendPort) const override
+	{
+		if (bAppendPort)
+		{
+			return FString::Printf(TEXT("%lld:%d"), SteamId.UniqueNetId, SteamChannel);
+		}
+		else
+		{
+			return FString::Printf(TEXT("%lld"), SteamId.UniqueNetId);
+		}
+	}
+
+	/**
+	 * Compares two internet ip addresses for equality
+	 *
+	 * @param Other the address to compare against
+	 */
+	virtual bool operator==(const FInternetAddr& Other) const override
+	{
+		FInternetAddrSteam& SteamOther = (FInternetAddrSteam&)Other;
+		return SteamId == SteamOther.SteamId && SteamChannel == SteamOther.SteamChannel;
+	}
+
+	bool operator!=(const FInternetAddrSteam& Other) const
+	{
+		return !(FInternetAddrSteam::operator==(Other));
+	}
+
+	/**
+	 * Is this a well formed internet address
+	 *
+	 * @return true if a valid IP, false otherwise
+	 */
+	virtual bool IsValid() const override
+	{
+		return SteamId.IsValid();
+	}
+};
diff --git a/Engine/Plugins/Online/OnlineSubsystemSteam/Source/Private/NboSerializerSteam.h b/Engine/Plugins/Online/OnlineSubsystemSteam/Source/Private/NboSerializerSteam.h
new file mode 100644
--- /dev/null
+++ b/Engine/Plugins/Online/OnlineSubsystemSteam/Source/Private/NboSerializerSteam.h
@@ -0,0 +1,82 @@
+// Copyright 1998-2016 Epic Games, Inc. All Rights Reserved.
+
+#pragma once
+
+#include "NboSerializer.h"
+#include "OnlineSubsystemSteamTypes.h"
+
+/**
+ * Serializes data in network byte order form into a buffer
+ */
+class FNboSerializeToBufferSteam : public FNboSerializeToBuffer
+{
+public:
+	/** Default constructor zeros num bytes*/
+	FNboSerializeToBufferSteam() :
+		FNboSerializeToBuffer(512)
+	{
+	}
+
+	/** Constructor specifying the size to use */
+	FNboSerializeToBufferSteam(uint32 Size) :
+		FNboSerializeToBuffer(Size)
+	{
+	}
+
+	/**
+	 * Adds Steam session info to the buffer
+	 */
+ 	friend inline FNboSerializeToBufferSteam& operator<<(FNboSerializeToBufferSteam& Ar, const FOnlineSessionInfoSteam& SessionInfo)
+ 	{
+		check(SessionInfo.HostAddr.IsValid());
+		// Skip SessionType (assigned at creation)
+		Ar << *SessionInfo.HostAddr;
+		Ar << SessionInfo.SessionId;
+		return Ar;
+ 	}
+
+	/**
+	 * Adds Steam Unique Id to the buffer
+	 */
+	friend inline FNboSerializeToBufferSteam& operator<<(FNboSerializeToBufferSteam& Ar, const FUniqueNetIdSteam& UniqueId)
+	{
+		Ar << UniqueId.UniqueNetId;
+		return Ar;
+	}
+};
+
+/**
+ * Class used to write data into packets for sending via system link
+ */
+class FNboSerializeFromBufferSteam : public FNboSerializeFromBuffer
+{
+public:
+	/**
+	 * Initializes the buffer, size, and zeros the read offset
+	 */
+	FNboSerializeFromBufferSteam(uint8* Packet,int32 Length) :
+		FNboSerializeFromBuffer(Packet,Length)
+	{
+	}
+
+	/**
+	 * Reads Steam session info from the buffer
+	 */
+ 	friend inline FNboSerializeFromBufferSteam& operator>>(FNboSerializeFromBufferSteam& Ar, FOnlineSessionInfoSteam& SessionInfo)
+ 	{
+		check(SessionInfo.HostAddr.IsValid());
+		// Skip SessionType (assigned at creation)
+		Ar >> *SessionInfo.HostAddr;
+		Ar >> SessionInfo.SessionId; 
+		return Ar;
+ 	}
+
+	/**
+	 * Reads Steam Unique Id from the buffer
+	 */
+	friend inline FNboSerializeFromBufferSteam& operator>>(FNboSerializeFromBufferSteam& Ar, FUniqueNetIdSteam& UniqueId)
+	{
+		Ar >> UniqueId.UniqueNetId;
+		return Ar;
+	}
+};
\ No newline at end of file
diff --git a/Engine/Plugins/Online/OnlineSubsystemSteam/Source/Private/OnlineAsyncTaskManagerSteam.cpp b/Engine/Plugins/Online/OnlineSubsystemSteam/Source/Private/OnlineAsyncTaskManagerSteam.cpp
new file mode 100644
--- /dev/null
+++ b/Engine/Plugins/Online/OnlineSubsystemSteam/Source/Private/OnlineAsyncTaskManagerSteam.cpp
@@ -0,0 +1,1057 @@
+// Copyright 1998-2016 Epic Games, Inc. All Rights Reserved.
+
+#include "OnlineSubsystemSteamPrivatePCH.h"
+#include "OnlineAsyncTaskManagerSteam.h"
+#include "OnlineSubsystemSteam.h"
+#include "OnlineSessionInterfaceSteam.h"
+#include "OnlineSessionAsyncLobbySteam.h"
+#include "OnlineSessionAsyncServerSteam.h"
+#include "OnlineSharedCloudInterfaceSteam.h"
+#include "OnlineLeaderboardInterfaceSteam.h"
+#include "OnlineExternalUIInterfaceSteam.h"
+#include "IPAddressSteam.h"
+#include "SocketSubsystemSteam.h"
+#include "SocketsSteam.h"
+#include "SteamUtilities.h"
+
+void FOnlineAsyncTaskManagerSteam::OnlineTick()
+{
+	check(SteamSubsystem);
+	check(FPlatformTLS::GetCurrentThreadId() == OnlineThreadId);
+
+	if (SteamSubsystem->IsSteamClientAvailable())
+	{
+		SteamAPI_RunCallbacks();
+	}
+
+	if (SteamSubsystem->IsSteamServerAvailable())
+	{
+		SteamGameServer_RunCallbacks();
+	}
+}
+
+/**
+ *	Event triggered by Steam backend when a user attempts JIP or accepts an invite request (via Steam client)
+ *
+ * @param CallbackData All the valid data from Steam related to this event 
+ */
+void FOnlineAsyncTaskManagerSteam::OnInviteAccepted(GameRichPresenceJoinRequested_t* CallbackData)
+{
+	FOnlineAsyncEventSteamInviteAccepted* NewEvent = 
+		new FOnlineAsyncEventSteamInviteAccepted(SteamSubsystem, FUniqueNetIdSteam(CallbackData->m_steamIDFriend), UTF8_TO_TCHAR(CallbackData->m_rgchConnect));
+	UE_LOG_ONLINE(Verbose, TEXT("%s"), *NewEvent->ToString());
+	AddToOutQueue(NewEvent);
+}
+
+/**
+ *	Event triggered by Steam backend when a user attempts JIP (via Steam client) or accepts an invite request (via Steam client)
+ *
+ * @param CallbackData All the valid data from Steam related to this event 
+ */
+void FOnlineAsyncTaskManagerSteam::OnLobbyInviteAccepted(GameLobbyJoinRequested_t* CallbackData)
+{
+	if (CallbackData->m_steamIDLobby.IsLobby())
+	{
+		FUniqueNetIdSteam LobbyId(CallbackData->m_steamIDLobby);
+
+		FOnlineSessionSteamPtr SessionInt = StaticCastSharedPtr<FOnlineSessionSteam>(SteamSubsystem->GetSessionInterface());
+		if (SessionInt.IsValid() && !SessionInt->IsMemberOfLobby(LobbyId))
+		{
+			FOnlineAsyncEventSteamLobbyInviteAccepted* NewEvent = 
+				new FOnlineAsyncEventSteamLobbyInviteAccepted(SteamSubsystem, FUniqueNetIdSteam(CallbackData->m_steamIDFriend), LobbyId);
+			UE_LOG_ONLINE(Verbose, TEXT("%s"), *NewEvent->ToString());
+			AddToOutQueue(NewEvent);
+		}
+		else
+		{
+			UE_LOG_ONLINE(Warning, TEXT("Attempting to accept invite to lobby user is already in, ignoring."));
+		}
+	}
+	else
+	{
+		UE_LOG_ONLINE(Warning, TEXT("OnLobbyInviteAccepted: Invalid LobbyId received."));
+	}
+}
+
+/**
+ * Notification event from Steam that the lobby state has changed (users joining/leaving)
+ */
+class FOnlineAsyncEventSteamLobbyEnter : public FOnlineAsyncEvent<FOnlineSubsystemSteam>
+{
+private:
+	
+	/** Lobby enter state information */
+	LobbyEnter_t CallbackResults;
+
+	/** Hidden on purpose */
+	FOnlineAsyncEventSteamLobbyEnter() :
+		FOnlineAsyncEvent(NULL)
+	{
+		FMemory::Memzero(CallbackResults);
+	}
+
+public:
+
+	FOnlineAsyncEventSteamLobbyEnter(FOnlineSubsystemSteam* InSubsystem, const LobbyEnter_t& InResults) :
+		FOnlineAsyncEvent(InSubsystem),
+		CallbackResults(InResults)
+	{
+	}
+	
+	virtual ~FOnlineAsyncEventSteamLobbyEnter() 
+	{
+	}
+
+	/**
+	 *	Get a human readable description of task
+	 */
+	virtual FString ToString() const override
+	{
+		return FString::Printf(TEXT("FOnlineAsyncEventSteamLobbyEnter LobbyId: %s Result: %s"), 
+			*FUniqueNetIdSteam(CallbackResults.m_ulSteamIDLobby).ToDebugString(),
+			*SteamChatRoomEnterResponseString((EChatRoomEnterResponse)CallbackResults.m_EChatRoomEnterResponse));	
+	}
+
+	/**
+	 * Give the async task a chance to marshal its data back to the game thread
+	 * Can only be called on the game thread by the async task manager
+	 */
+	virtual void Finalize() override
+	{
+		FOnlineSessionSteamPtr SessionInt = StaticCastSharedPtr<FOnlineSessionSteam>(Subsystem->GetSessionInterface());
+		if (SessionInt.IsValid())
+		{
+			FUniqueNetIdSteam LobbyId(CallbackResults.m_ulSteamIDLobby);
+			FNamedOnlineSession* Session = SessionInt->GetNamedSessionFromLobbyId(LobbyId);
+			if (!Session)
+			{
+				UE_LOG_ONLINE(Warning, TEXT("Entered lobby %s, but not found in sessions list"), *LobbyId.ToDebugString());
+			}
+		}
+	}
+};
+
+/**
+ *	Event triggered by Steam backend when a user joins a lobby
+ *
+ * @param CallbackData All the valid data from Steam related to this event 
+ */
+void FOnlineAsyncTaskManagerSteam::OnLobbyEnter(LobbyEnter_t* CallbackData)
+{
+	// The owner of the created lobby shouldn't need this information
+	if (SteamMatchmaking()->GetLobbyOwner(CallbackData->m_ulSteamIDLobby) != SteamUser()->GetSteamID())
+	{
+		FOnlineAsyncEventSteamLobbyEnter* NewEvent = new FOnlineAsyncEventSteamLobbyEnter(SteamSubsystem, *CallbackData);
+		UE_LOG_ONLINE(Verbose, TEXT("%s"), *NewEvent->ToString());
+		AddToOutQueue(NewEvent);
+	}
+}
+
+/**
+ * Notification event from Steam that the lobby state has changed (users joining/leaving)
+ */
+class FOnlineAsyncEventSteamLobbyChatUpdate : public FOnlineAsyncEvent<FOnlineSubsystemSteam>
+{
+private:
+	
+	/** Lobby chat state information */
+	LobbyChatUpdate_t CallbackResults;
+
+	/** Hidden on purpose */
+	FOnlineAsyncEventSteamLobbyChatUpdate() :
+		FOnlineAsyncEvent(NULL)
+	{
+		FMemory::Memzero(CallbackResults);
+	}
+
+public:
+
+	FOnlineAsyncEventSteamLobbyChatUpdate(FOnlineSubsystemSteam* InSubsystem, const LobbyChatUpdate_t& InResults) :
+		FOnlineAsyncEvent(InSubsystem),
+		CallbackResults(InResults)
+	{
+	}
+	
+	virtual ~FOnlineAsyncEventSteamLobbyChatUpdate() 
+	{
+	}
+
+	/**
+	 *	Get a human readable description of task
+	 */
+	virtual FString ToString() const override
+	{
+		return FString::Printf(TEXT("FOnlineAsyncEventSteamLobbyChatUpdate User: %s Instigator: %s Result: %s"), 
+			*FUniqueNetIdSteam(CallbackResults.m_ulSteamIDUserChanged).ToDebugString(),
+			*FUniqueNetIdSteam(CallbackResults.m_ulSteamIDMakingChange).ToDebugString(),
+			*SteamChatMemberStateChangeString((EChatMemberStateChange)CallbackResults.m_rgfChatMemberStateChange));	
+	}
+
+	/**
+	 * Give the async task a chance to marshal its data back to the game thread
+	 * Can only be called on the game thread by the async task manager
+	 */
+	virtual void Finalize() override
+	{
+		FOnlineSessionSteamPtr SessionInt = StaticCastSharedPtr<FOnlineSessionSteam>(Subsystem->GetSessionInterface());
+		if (SessionInt.IsValid())
+		{
+			FUniqueNetIdSteam LobbyId(CallbackResults.m_ulSteamIDLobby);
+			// Lobby data update for existing session
+			FNamedOnlineSession* Session = SessionInt->GetNamedSessionFromLobbyId(LobbyId);
+			if (Session)
+			{
+				// Recreate the lobby member list
+				if (!FillMembersFromLobbyData(LobbyId, *Session))
+				{
+					UE_LOG_ONLINE(Warning, TEXT("Failed to parse session %s member update %s"), *Session->SessionName.ToString(), *LobbyId.ToDebugString());
+				}
+			}
+			else
+			{
+				UE_LOG_ONLINE(Warning, TEXT("Received lobby chat update %s, but not found in sessions list"), *LobbyId.ToDebugString());
+			}
+		}
+	}
+};
+
+/**
+ *	Event triggered by Steam backend when a user joins a lobby
+ *
+ * @param CallbackData All the valid data from Steam related to this event 
+ */
+void FOnlineAsyncTaskManagerSteam::OnLobbyChatUpdate(LobbyChatUpdate_t* CallbackData)
+{
+	FOnlineAsyncEventSteamLobbyChatUpdate* NewEvent = new FOnlineAsyncEventSteamLobbyChatUpdate(SteamSubsystem, *CallbackData);
+	UE_LOG_ONLINE(Verbose, TEXT("%s"), *NewEvent->ToString());
+	AddToOutQueue(NewEvent);
+}
+
+/**
+ * Notification event from Steam when new lobby data is available for the given lobby 
+ */
+class FOnlineAsyncEventSteamLobbyUpdate : public FOnlineAsyncEvent<FOnlineSubsystemSteam> 
+{
+private:
+
+	/** Id of lobby to update */
+	FUniqueNetIdSteam LobbyId;
+	/** Hidden on purpose */
+	FOnlineAsyncEventSteamLobbyUpdate() 
+	{}
+
+public:
+	FOnlineAsyncEventSteamLobbyUpdate(FOnlineSubsystemSteam* InSubsystem, const FUniqueNetIdSteam& InLobbyId) :
+		FOnlineAsyncEvent(InSubsystem),
+		LobbyId(InLobbyId)
+	{
+	}
+
+	virtual ~FOnlineAsyncEventSteamLobbyUpdate() 
+	{
+	}
+
+	/**
+	 *	Get a human readable description of task
+	 */
+	virtual FString ToString() const override
+	{
+		return FString::Printf(TEXT("FOnlineAsyncEventSteamLobbyUpdate LobbyId: %s"), *LobbyId.ToDebugString());	
+	}
+
+	/**
+	 * Give the async task a chance to marshal its data back to the game thread
+	 * Can only be called on the game thread by the async task manager
+	 */
+	virtual void Finalize() override
+	{
+		FOnlineAsyncEvent::Finalize();
+
+		// Searching for lobbies case (NULL CurrentSessionSearch implies no active search query)
+		FOnlineSessionSteamPtr SessionInt = StaticCastSharedPtr<FOnlineSessionSteam>(Subsystem->GetSessionInterface());
+		if (SessionInt.IsValid() && SessionInt->CurrentSessionSearch.IsValid() && SessionInt->CurrentSessionSearch->SearchState == EOnlineAsyncTaskState::InProgress)
+		{
+			// Add this lobby as available for adding to search results
+			SessionInt->PendingSearchLobbyIds.AddUnique(LobbyId);
+		}
+		else
+		{
+			// Lobby data update for existing session
+			FNamedOnlineSession* Session = SessionInt->GetNamedSessionFromLobbyId(LobbyId);
+			if (Session)
+			{
+				// Make sure the session has all the valid session data
+				if (!FillSessionFromLobbyData(LobbyId, *Session) ||
+					!FillMembersFromLobbyData(LobbyId, *Session))
+				{
+					UE_LOG_ONLINE(Warning, TEXT("Failed to parse session %s lobby update %s"), *Session->SessionName.ToString(), *LobbyId.ToDebugString());
+				}
+			}
+			else
+			{
+				UE_LOG_ONLINE(Warning, TEXT("Received lobby update %s, but not found in sessions list"), *LobbyId.ToDebugString());
+			}
+		}
+	}
+};
+
+/**
+ *	Event triggered by Steam backend when new lobby data is available for the given lobby 
+ * Can occur any time host calls SetLobbyData or while searching for lobbies (calls to RequestLobbyData)
+ *
+ * @param CallbackData All the valid data from Steam related to this event 
+ */
+void FOnlineAsyncTaskManagerSteam::OnLobbyDataUpdate(LobbyDataUpdate_t* CallbackData)
+{
+	// Equivalent lobby ids implies it is lobby data that has updated
+	if (CallbackData->m_ulSteamIDLobby == CallbackData->m_ulSteamIDMember)
+	{
+		FUniqueNetIdSteam LobbyId(CallbackData->m_ulSteamIDLobby);
+		if (!CallbackData->m_bSuccess)
+		{
+			// CallbackData->m_bSuccess indicates LobbyID has shut down since 
+			// the result was returned but we have to keep the array size in sync
+			UE_LOG_ONLINE(Verbose, TEXT("Lobby %s is no longer available."), *LobbyId.ToDebugString());
+		}
+
+		ISteamMatchmaking* SteamMatchmakingPtr = SteamMatchmaking();
+		check(SteamMatchmakingPtr);
+
+		// The owner of the created lobby shouldn't need this information
+		if (SteamMatchmakingPtr->GetLobbyOwner(CallbackData->m_ulSteamIDLobby) != SteamUser()->GetSteamID())
+		{
+			FOnlineAsyncEventSteamLobbyUpdate* NewEvent = new FOnlineAsyncEventSteamLobbyUpdate(SteamSubsystem, LobbyId);
+			UE_LOG_ONLINE(Verbose, TEXT("%s"), *NewEvent->ToString());
+			AddToOutQueue(NewEvent);
+		}
+	}
+	else
+	{
+		// @TODO ONLINE - Player data update
+		//const TCHAR* Key = MemberKeyList(i);
+		//const ANSICHAR* MemberValue = SteamMatchmakingPtr->GetLobbyMemberData(CallbackData->m_ulSteamIDLobby, CallbackData->m_ulSteamIDMember, TCHAR_TO_UTF8(Key));
+	}
+
+	// @TODO ONLINE - SetLobbyOwner triggers this call also
+}
+
+/**
+ * Notification event from Steam that a given user's
+ * stats/achievements data has been downloaded from the server
+ */
+class FOnlineAsyncEventSteamStatsReceived : public FOnlineAsyncEvent<FOnlineSubsystemSteam> 
+{
+private:
+
+	/** User this data is for */
+	const FUniqueNetIdSteam UserId;
+	/** Result of the download */
+	EResult StatsReceivedResult;
+
+	/** Hidden on purpose */
+	FOnlineAsyncEventSteamStatsReceived() :
+		UserId(0),
+		StatsReceivedResult(k_EResultOK)
+	{}
+
+public:
+	FOnlineAsyncEventSteamStatsReceived(FOnlineSubsystemSteam* InSubsystem, const FUniqueNetIdSteam& InUserId, EResult InResult) :
+		FOnlineAsyncEvent(InSubsystem),
+		UserId(InUserId),
+		StatsReceivedResult(InResult)
+	{
+	}
+
+	virtual ~FOnlineAsyncEventSteamStatsReceived() 
+	{
+	}
+
+	/**
+	 *	Get a human readable description of task
+	 */
+	virtual FString ToString() const override
+	{
+		return FString::Printf(TEXT("FOnlineAsyncEventSteamStatsReceived bWasSuccessful: %d User: %s Result: %s"), (StatsReceivedResult == k_EResultOK) ? 1 : 0, *UserId.ToDebugString(), *SteamResultString(StatsReceivedResult));	
+	}
+};
+
+/**
+ *	Event triggered from Steam when the current user's stats have been downloaded from the backend
+ *	Possible that the result fails if they have no data for the current game
+ *
+ * @param CallbackData All the valid data from Steam related to this event 
+ */
+void FOnlineAsyncTaskManagerSteam::OnUserStatsReceived(UserStatsReceived_t* CallbackData)
+{
+	const CGameID GameID(SteamSubsystem->GetSteamAppId());
+	if (GameID.ToUint64() == CallbackData->m_nGameID)
+	{
+		FUniqueNetIdSteam UserId(CallbackData->m_steamIDUser);
+		if (CallbackData->m_eResult != k_EResultOK)
+		{
+			if (CallbackData->m_eResult == k_EResultFail)
+			{
+				UE_LOG_ONLINE(Warning, TEXT("Failed to obtain steam user stats, user: %s has no stats entries"), *UserId.ToDebugString());
+			}
+			else
+			{
+				UE_LOG_ONLINE(Warning, TEXT("Failed to obtain steam user stats, user: %s error: %s"), *UserId.ToDebugString(), 
+					*SteamResultString(CallbackData->m_eResult));
+			}
+		}
+
+		FOnlineAsyncEventSteamStatsReceived* NewEvent = new FOnlineAsyncEventSteamStatsReceived(SteamSubsystem, UserId, CallbackData->m_eResult);
+		UE_LOG_ONLINE(Verbose, TEXT("%s"), *NewEvent->ToString());
+		AddToOutQueue(NewEvent);
+	}
+	else
+	{
+		UE_LOG_ONLINE(Warning, TEXT("Obtained steam user stats, but for wrong game! Ignoring."));
+	}
+}
+
+/**
+ * Notification event from Steam that the currently logged in user's
+ * stats/achievements data has been stored with the server
+ */
+class FOnlineAsyncEventSteamStatsStored : public FOnlineAsyncEvent<FOnlineSubsystemSteam>
+{
+private:
+	
+	/** User this data is for */
+	const FUniqueNetIdSteam UserId;
+	/** Result of the download */
+	EResult StatsStoredResult;
+
+	/** Hidden on purpose */
+	FOnlineAsyncEventSteamStatsStored() :
+		FOnlineAsyncEvent(NULL),
+		UserId(0),
+		StatsStoredResult(k_EResultOK)
+	{
+	}
+
+public:
+
+	FOnlineAsyncEventSteamStatsStored(FOnlineSubsystemSteam* InSubsystem, const FUniqueNetIdSteam& InUserId, EResult InResult) :
+		FOnlineAsyncEvent(InSubsystem),
+		UserId(InUserId),
+		StatsStoredResult(InResult)
+	{
+	}
+	
+	virtual ~FOnlineAsyncEventSteamStatsStored() 
+	{
+	}
+
+	/**
+	 *	Get a human readable description of task
+	 */
+	virtual FString ToString() const override
+	{
+		return FString::Printf(TEXT("FOnlineAsyncEventSteamStatsStored bWasSuccessful: %d User: %s Result: %s"), (StatsStoredResult == k_EResultOK) ? 1 : 0, *UserId.ToDebugString(), *SteamResultString(StatsStoredResult));	
+	}
+
+	/**
+	 * Give the async task a chance to marshal its data back to the game thread
+	 * Can only be called on the game thread by the async task manager
+	 */
+	virtual void Finalize() override
+	{
+		FOnlineAsyncEvent::Finalize();
+		FOnlineLeaderboardsSteamPtr Leaderboards = StaticCastSharedPtr<FOnlineLeaderboardsSteam>(Subsystem->GetLeaderboardsInterface());
+		Leaderboards->UserStatsStoreStatsFinishedDelegate.ExecuteIfBound(StatsStoredResult == k_EResultOK ? EOnlineAsyncTaskState::Done : EOnlineAsyncTaskState::Failed);
+	}
+};
+
+/**
+ *	Event triggered from Steam when the current user's stats have been stored with the backend
+ *	Possible that the result fails with "Invalid Param" meaning the stats went out of range or were out of date
+ *  New stats are downloaded in this case and need to be re-evaluated
+ *
+ * @param CallbackData All the valid data from Steam related to this event 
+ */
+void FOnlineAsyncTaskManagerSteam::OnUserStatsStored(UserStatsStored_t* CallbackData)
+{
+	const CGameID GameID(SteamSubsystem->GetSteamAppId());
+	if (GameID.ToUint64() == CallbackData->m_nGameID)
+	{
+		// Only the current user comes through this way (other user's stats are stored via GameServerStats)
+		FUniqueNetIdSteam UserId(SteamUser()->GetSteamID());
+		if (CallbackData->m_eResult != k_EResultOK)
+		{
+			if (CallbackData->m_eResult == k_EResultInvalidParam)
+			{
+				UE_LOG_ONLINE(Warning, TEXT("Invalid stats data set, stats have been reverted to state prior to last write."));
+			}
+			else
+			{
+				UE_LOG_ONLINE(Warning, TEXT("Failed to store steam user stats, error: %s"), *SteamResultString(CallbackData->m_eResult));
+			}
+		}
+
+		FOnlineAsyncEventSteamStatsStored* NewEvent = new FOnlineAsyncEventSteamStatsStored(SteamSubsystem, UserId, CallbackData->m_eResult);
+		UE_LOG_ONLINE(Verbose, TEXT("%s"), *NewEvent->ToString());
+		AddToOutQueue(NewEvent);
+	}
+	else
+	{
+		UE_LOG_ONLINE(Warning, TEXT("Stored steam user stats, but for wrong game! Ignoring."));
+	}
+}
+
+/**
+ * Notification event from Steam that the currently logged in user's
+ * stats/achievements data has been stored with the server
+ * FROM VALVE: Steam stats for other users are kept in an LRU with a max queue length of 100
+ */
+class FOnlineAsyncEventSteamStatsUnloaded : public FOnlineAsyncEvent<FOnlineSubsystemSteam>
+{
+private:
+	/** User whose data has been unloaded */
+	FUniqueNetIdSteam UserId;
+
+	/** Hidden on purpose */
+	FOnlineAsyncEventSteamStatsUnloaded() : 
+		UserId(0)
+	{
+	}
+
+public:
+
+	FOnlineAsyncEventSteamStatsUnloaded(FOnlineSubsystemSteam* InSubsystem, const FUniqueNetIdSteam& InUserId) :
+		FOnlineAsyncEvent(InSubsystem),
+		UserId(InUserId)
+	{
+	}
+	
+	virtual ~FOnlineAsyncEventSteamStatsUnloaded() 
+	{
+	}
+
+	/**
+	 *	Get a human readable description of task
+	 */
+	virtual FString ToString() const override
+	{
+		return FString::Printf(TEXT("FOnlineAsyncEventSteamStatsUnloaded UserId: %s"), *UserId.ToDebugString());	
+	}
+};
+
+/**
+ *	Event triggered from Steam when a previously requested user's stats have been purged in LRU fashion
+ * Requesting the data an additional time will bring the data back
+ *
+ * @param CallbackData All the valid data from Steam related to this event 
+ */
+void FOnlineAsyncTaskManagerSteam::OnUserStatsUnloaded(UserStatsUnloaded_t* CallbackData)
+{
+	FOnlineAsyncEventSteamStatsUnloaded* NewEvent = new FOnlineAsyncEventSteamStatsUnloaded(SteamSubsystem, FUniqueNetIdSteam(CallbackData->m_steamIDUser));
+	UE_LOG_ONLINE(Verbose, TEXT("%s"), *NewEvent->ToString());
+	AddToOutQueue(NewEvent);
+}
+
+/**
+ *	Delegate registered with Steam to trigger when the Steam Overlay is activated
+ *
+ * @param CallbackData - Steam struct containing state of the Steam Overlay
+ */
+void FOnlineAsyncTaskManagerSteam::OnExternalUITriggered(GameOverlayActivated_t* CallbackData)
+{
+	FOnlineAsyncEventSteamExternalUITriggered* NewEvent = new FOnlineAsyncEventSteamExternalUITriggered(SteamSubsystem, (CallbackData->m_bActive != 0) ? true : false);
+	UE_LOG_ONLINE(Verbose, TEXT("%s"), *NewEvent->ToString());
+	AddToOutQueue(NewEvent);
+}
+
+/**
+ * Notification event from Steam that server session connection has changed state 
+ */
+class FOnlineAsyncEventSteamServerConnectionState : public FOnlineAsyncEvent<FOnlineSubsystemSteam>
+{
+	/** Connection state change */
+	const EOnlineServerConnectionStatus::Type ConnectionState;
+
+	/** Hidden on purpose */
+	FOnlineAsyncEventSteamServerConnectionState() :
+		ConnectionState(EOnlineServerConnectionStatus::NotConnected)
+	{
+	}
+
+public:
+
+	FOnlineAsyncEventSteamServerConnectionState(FOnlineSubsystemSteam* InSubsystem, EOnlineServerConnectionStatus::Type InConnectionState) :
+		FOnlineAsyncEvent(InSubsystem),
+		ConnectionState(InConnectionState)
+	{
+	}
+	
+	virtual ~FOnlineAsyncEventSteamServerConnectionState() 
+	{
+	}
+
+	/**
+	 *	Get a human readable description of task
+	 */
+	virtual FString ToString() const override
+	{
+		return FString::Printf(TEXT("FOnlineAsyncEventSteamServerConnectionState StateChange: %s"), EOnlineServerConnectionStatus::ToString(ConnectionState));	
+	}
+
+	/**
+	 * Give the async task a chance to marshal its data back to the game thread
+	 * Can only be called on the game thread by the async task manager
+	 */
+	virtual void Finalize() override
+	{
+		Subsystem->TriggerOnConnectionStatusChangedDelegates(EOnlineServerConnectionStatus::Normal, ConnectionState);
+	}
+};
+
+/**
+ *	Client API version of the connected to Steam callback (only called in case of a Steam backend disconnect and then reconnect)
+ */
+void FOnlineAsyncTaskManagerSteam::OnSteamServersConnected(SteamServersConnected_t* CallbackData)
+{
+	FOnlineAsyncEventSteamServerConnectionState* NewEvent = new FOnlineAsyncEventSteamServerConnectionState(SteamSubsystem, EOnlineServerConnectionStatus::Connected);
+	UE_LOG_ONLINE(Verbose, TEXT("%s"), *NewEvent->ToString());
+	AddToOutQueue(NewEvent);
+}
+
+/**
+ *	Client API version of the disconnected to Steam callback 
+ */
+void FOnlineAsyncTaskManagerSteam::OnSteamServersDisconnected(SteamServersDisconnected_t* CallbackData)
+{
+	FOnlineAsyncEventSteamServerConnectionState* NewEvent = new FOnlineAsyncEventSteamServerConnectionState(SteamSubsystem, SteamConnectionResult(CallbackData->m_eResult));
+	UE_LOG_ONLINE(Verbose, TEXT("%s"), *NewEvent->ToString());
+	AddToOutQueue(NewEvent);
+}
+
+/**
+ * Notification event from Steam that server session has connected with the master server
+ */
+class FOnlineAsyncEventSteamServerConnectedGS : public FOnlineAsyncEvent<FOnlineSubsystemSteam>
+{
+	/** Newly assigned server id */
+	const FUniqueNetIdSteam ServerId;
+
+	/** Hidden on purpose */
+	FOnlineAsyncEventSteamServerConnectedGS() :
+		ServerId(uint64(0))
+	{
+	}
+
+public:
+
+	FOnlineAsyncEventSteamServerConnectedGS(FOnlineSubsystemSteam* InSubsystem, const FUniqueNetIdSteam& InServerId) :
+		FOnlineAsyncEvent(InSubsystem),
+		ServerId(InServerId)
+	{
+	}
+	
+	virtual ~FOnlineAsyncEventSteamServerConnectedGS() 
+	{
+	}
+
+	/**
+	 *	Get a human readable description of task
+	 */
+	virtual FString ToString() const override
+	{
+		return FString::Printf(TEXT("FOnlineAsyncEventSteamServerConnectedGS ServerId: %s"), *ServerId.ToDebugString());	
+	}
+
+	/**
+	 * Give the async task a chance to marshal its data back to the game thread
+	 * Can only be called on the game thread by the async task manager
+	 */
+	virtual void Finalize() override
+	{
+		FOnlineSessionSteamPtr SessionInt = StaticCastSharedPtr<FOnlineSessionSteam>(Subsystem->GetSessionInterface());
+		if (SessionInt.IsValid())
+		{
+			SessionInt->bSteamworksGameServerConnected = true;
+			SessionInt->GameServerSteamId = MakeShareable(new FUniqueNetIdSteam(ServerId));
+
+			FSocketSubsystemSteam* SocketSubsystem = (FSocketSubsystemSteam*)ISocketSubsystem::Get(STEAM_SUBSYSTEM);
+			if (SocketSubsystem)
+			{			
+				SocketSubsystem->FixupSockets(*SessionInt->GameServerSteamId);
+			}
+		}
+
+		// log on is not finished until OnPolicyResponse() is called
+	}
+};
+
+/**
+ *	GameServer API version of connected to Steam backend callback, 
+ *  initiated by SteamGameServers()->LogOnAnonymous()
+ */
+void FOnlineAsyncTaskManagerSteam::OnSteamServersConnectedGS(SteamServersConnected_t* CallbackData)
+{
+	FOnlineAsyncEventSteamServerConnectedGS* NewEvent = new FOnlineAsyncEventSteamServerConnectedGS(SteamSubsystem, FUniqueNetIdSteam(SteamGameServer()->GetSteamID()));
+	UE_LOG_ONLINE(Verbose, TEXT("%s"), *NewEvent->ToString());
+	AddToOutQueue(NewEvent);
+}
+
+/**
+ * Notification event from Steam that server session has been disconnected with the master server
+ */
+class FOnlineAsyncEventSteamServerDisconnectedGS : public FOnlineAsyncEvent<FOnlineSubsystemSteam>
+{
+private:
+
+	/** Callback data */
+	SteamServersDisconnected_t CallbackResults;
+
+	/** Hidden on purpose */
+	FOnlineAsyncEventSteamServerDisconnectedGS()
+	{
+	}
+
+public:
+
+	FOnlineAsyncEventSteamServerDisconnectedGS(FOnlineSubsystemSteam* InSubsystem, SteamServersDisconnected_t& InResults) :
+		FOnlineAsyncEvent(InSubsystem),
+		CallbackResults(InResults)
+	{
+	}
+	
+	virtual ~FOnlineAsyncEventSteamServerDisconnectedGS() 
+	{
+	}
+
+	/**
+	 *	Get a human readable description of task
+	 */
+	virtual FString ToString() const override
+	{
+		return FString::Printf(TEXT("FOnlineAsyncEventSteamServerDisconnectedGS Result: %s"), *SteamResultString(CallbackResults.m_eResult));	
+	}
+
+	/**
+	 * Give the async task a chance to marshal its data back to the game thread
+	 * Can only be called on the game thread by the async task manager
+	 */
+	virtual void Finalize() override
+	{
+		bool bTriggerConnectionStatusUpdate = true;
+		FOnlineSessionSteamPtr SessionInt = StaticCastSharedPtr<FOnlineSessionSteam>(Subsystem->GetSessionInterface());
+		if (SessionInt.IsValid())
+		{
+			SessionInt->bSteamworksGameServerConnected = false;
+			SessionInt->GameServerSteamId = NULL;
+
+			// Don't trigger the delegates if a DestroySession() call was made
+			FNamedOnlineSession* Session = SessionInt->GetGameServerSession();
+			if (Session && Session->SessionState == EOnlineSessionState::Destroying)
+			{
+				bTriggerConnectionStatusUpdate = false;
+			}
+		}
+		
+		if (bTriggerConnectionStatusUpdate)
+		{
+			EOnlineServerConnectionStatus::Type ConnectionState = SteamConnectionResult(CallbackResults.m_eResult);
+			Subsystem->TriggerOnConnectionStatusChangedDelegates(EOnlineServerConnectionStatus::Normal, ConnectionState);
+		}
+	}
+};
+
+/**
+ *	GameServer API version of disconnected from Steam backend callback
+ */
+void FOnlineAsyncTaskManagerSteam::OnSteamServersDisconnectedGS(SteamServersDisconnected_t* CallbackData)
+{
+	FOnlineAsyncEventSteamServerDisconnectedGS* NewEvent = new FOnlineAsyncEventSteamServerDisconnectedGS(SteamSubsystem, *CallbackData);
+	UE_LOG_ONLINE(Verbose, TEXT("%s"), *NewEvent->ToString());
+	AddToOutQueue(NewEvent);
+}
+
+/**
+ *	GameServer API version of disconnected from Steam backend callback
+ */
+void FOnlineAsyncTaskManagerSteam::OnSteamServersConnectFailureGS(SteamServerConnectFailure_t* CallbackData)
+{
+	UE_LOG_ONLINE(Warning, TEXT("Steam connection failure."));
+}
+
+/**
+ * Notification event from Steam that server session has been secured on the backend
+ */
+class FOnlineAsyncEventSteamServerPolicyResponseGS : public FOnlineAsyncEvent<FOnlineSubsystemSteam>
+{
+private:
+
+	/** Callback data */
+	GSPolicyResponse_t CallbackResults;
+
+	/** Hidden on purpose */
+	FOnlineAsyncEventSteamServerPolicyResponseGS()
+	{
+	}
+
+public:
+
+	FOnlineAsyncEventSteamServerPolicyResponseGS(FOnlineSubsystemSteam* InSubsystem, GSPolicyResponse_t& InResults) :
+		FOnlineAsyncEvent(InSubsystem),
+		CallbackResults(InResults)
+	{
+	}
+	
+	virtual ~FOnlineAsyncEventSteamServerPolicyResponseGS() 
+	{
+	}
+
+	/**
+	 *	Get a human readable description of task
+	 */
+	virtual FString ToString() const override
+	{
+		return FString::Printf(TEXT("FOnlineAsyncEventSteamServerPolicyResponseGS Secure: %d"), CallbackResults.m_bSecure);	
+	}
+
+	/**
+	 * Give the async task a chance to marshal its data back to the game thread
+	 * Can only be called on the game thread by the async task manager
+	 */
+	virtual void Finalize() override
+	{
+		FOnlineSessionSteamPtr SessionInt = StaticCastSharedPtr<FOnlineSessionSteam>(Subsystem->GetSessionInterface());
+		if (SessionInt.IsValid())
+		{
+			SessionInt->bPolicyResponseReceived = true;
+			if (!SessionInt->bSteamworksGameServerConnected || !SessionInt->GameServerSteamId->IsValid())
+			{
+				UE_LOG_ONLINE(Warning, TEXT("Unexpected GSPolicyResponse callback"));
+			}
+		}
+	}
+};
+
+/**
+ * Notification event from Steam that server session has been secured
+ */
+void FOnlineAsyncTaskManagerSteam::OnPolicyResponseGS(GSPolicyResponse_t* CallbackData)
+{
+	FOnlineAsyncEventSteamServerPolicyResponseGS* NewEvent = new FOnlineAsyncEventSteamServerPolicyResponseGS(SteamSubsystem, *CallbackData);
+	UE_LOG_ONLINE(Verbose, TEXT("%s"), *NewEvent->ToString());
+	AddToOutQueue(NewEvent);
+}
+
+/**
+ * Notification event from Steam that a P2P connection has been requested from a remote user
+ */
+class FOnlineAsyncEventSteamConnectionRequest : public FOnlineAsyncEvent<FOnlineSubsystemSteam>
+{
+private:
+	
+	/** Proper networking interface that this session is communicating on */
+	ISteamNetworking* SteamNetworkingPtr;
+	/** Callback data */
+	FUniqueNetIdSteam RemoteId;
+
+	/** Hidden on purpose */
+	FOnlineAsyncEventSteamConnectionRequest()
+	{
+	}
+
+public:
+
+	FOnlineAsyncEventSteamConnectionRequest(FOnlineSubsystemSteam* InSubsystem, ISteamNetworking* InSteamNetworkingPtr, const FUniqueNetIdSteam& InRemoteId) :
+		FOnlineAsyncEvent(InSubsystem),
+		SteamNetworkingPtr(InSteamNetworkingPtr),
+		RemoteId(InRemoteId)
+	{
+	}
+	
+	virtual ~FOnlineAsyncEventSteamConnectionRequest() 
+	{
+	}
+
+	/**
+	 *	Get a human readable description of task
+	 */
+	virtual FString ToString() const override
+	{
+		return FString::Printf(TEXT("FOnlineAsyncEventSteamConnectionRequest RemoteId: %s"), *RemoteId.ToDebugString());	
+	}
+
+	/**
+	 * Give the async task a chance to marshal its data back to the game thread
+	 * Can only be called on the game thread by the async task manager
+	 */
+	virtual void Finalize() override
+	{
+		FSocketSubsystemSteam* SocketSubsystem = (FSocketSubsystemSteam*)ISocketSubsystem::Get(STEAM_SUBSYSTEM);
+		if (SocketSubsystem)
+		{
+			if (!SocketSubsystem->AcceptP2PConnection(SteamNetworkingPtr, RemoteId))
+			{	
+				UE_LOG_ONLINE(Log, TEXT("Rejected P2P connection request from %s"), *RemoteId.ToDebugString());
+			}
+		}
+	}
+};
+
+/**
+ * Notification event from Steam that a P2P connection has failed
+ */
+class FOnlineAsyncEventSteamConnectionFailed : public FOnlineAsyncEvent<FOnlineSubsystemSteam>
+{
+private:
+
+	/** Callback data */
+	FUniqueNetIdSteam RemoteId;
+	/** Error reason */
+	EP2PSessionError ErrorCode;
+
+	/** Hidden on purpose */
+	FOnlineAsyncEventSteamConnectionFailed()
+	{
+	}
+
+public:
+
+	FOnlineAsyncEventSteamConnectionFailed(FOnlineSubsystemSteam* InSubsystem, const FUniqueNetIdSteam& InRemoteId, EP2PSessionError InErrorCode) :
+		FOnlineAsyncEvent(InSubsystem),
+		RemoteId(InRemoteId),
+		ErrorCode(InErrorCode)
+	{
+	}
+	
+	virtual ~FOnlineAsyncEventSteamConnectionFailed() 
+	{
+	}
+
+	/**
+	 *	Get a human readable description of task
+	 */
+	virtual FString ToString() const override
+	{
+		return FString::Printf(TEXT("FOnlineAsyncEventSteamConnectionFailed RemoteId: %s Reason: %s"), *RemoteId.ToDebugString(), *SteamP2PConnectError(ErrorCode));	
+	}
+
+	/**
+	 * Give the async task a chance to marshal its data back to the game thread
+	 * Can only be called on the game thread by the async task manager
+	 */
+	virtual void Finalize() override
+	{
+		// Mark the relevant sockets with this failure so they can properly notify higher level engine code
+		FSocketSubsystemSteam* SocketSubsystem = (FSocketSubsystemSteam*)ISocketSubsystem::Get(STEAM_SUBSYSTEM);
+		if (SocketSubsystem)
+		{
+			SocketSubsystem->ConnectFailure(RemoteId);
+		}
+	}
+};
+
+/**
+ * Notification event from Steam that a P2P connection request has been initiated from a remote connection
+ *
+ * @param CallbackData information about remote connection request
+ */
+void FOnlineAsyncTaskManagerSteam::OnP2PSessionRequest(P2PSessionRequest_t* CallbackData)
+{
+	FUniqueNetIdSteam RemoteId(CallbackData->m_steamIDRemote);
+	UE_LOG_ONLINE(Verbose, TEXT("Client connection request Id: %s"), *RemoteId.ToDebugString());
+
+	IOnlineSessionPtr SessionInt = SteamSubsystem->GetSessionInterface();
+	// Only accept connections if we have any expectation of being online
+	if (SessionInt.IsValid() && SessionInt->GetNumSessions() > 0)
+	{
+		FOnlineAsyncEventSteamConnectionRequest* NewEvent = new FOnlineAsyncEventSteamConnectionRequest(SteamSubsystem, SteamNetworking(), RemoteId);
+		UE_LOG_ONLINE(Verbose, TEXT("%s"), *NewEvent->ToString());
+		AddToOutQueue(NewEvent);
+	}
+}
+
+/**
+ * Notification event from Steam that a P2P remote connection has failed
+ * 
+ * @param CallbackData information about remote connection failure
+ */
+void FOnlineAsyncTaskManagerSteam::OnP2PSessionConnectFail(P2PSessionConnectFail_t* CallbackData)
+{
+	FUniqueNetIdSteam RemoteId(CallbackData->m_steamIDRemote);
+	FOnlineAsyncEventSteamConnectionFailed* NewEvent = new FOnlineAsyncEventSteamConnectionFailed(SteamSubsystem, RemoteId, (EP2PSessionError)CallbackData->m_eP2PSessionError);
+	UE_LOG_ONLINE(Verbose, TEXT("%s"), *NewEvent->ToString());
+	AddToOutQueue(NewEvent);
+}
+
+/**
+ * Notification event from Steam that a P2P connection request has been initiated from a remote connection
+ * (GameServer version)
+ * 
+ * @param CallbackData information about remote connection request
+ */
+void FOnlineAsyncTaskManagerSteam::OnP2PSessionRequestGS(P2PSessionRequest_t* CallbackData)
+{
+	FUniqueNetIdSteam RemoteId(CallbackData->m_steamIDRemote);
+	FOnlineAsyncEventSteamConnectionRequest* NewEvent = new FOnlineAsyncEventSteamConnectionRequest(SteamSubsystem, SteamGameServerNetworking(), RemoteId);
+	UE_LOG_ONLINE(Verbose, TEXT("%s"), *NewEvent->ToString());
+	AddToOutQueue(NewEvent);
+}
+
+/**
+ * Notification event from Steam that a P2P remote connection has failed
+ * (GameServer version)
+ * 
+ * @param CallbackData information about remote connection failure
+ */
+void FOnlineAsyncTaskManagerSteam::OnP2PSessionConnectFailGS(P2PSessionConnectFail_t* CallbackData)
+{
+	FUniqueNetIdSteam RemoteId(CallbackData->m_steamIDRemote);
+	FOnlineAsyncEventSteamConnectionFailed* NewEvent = new FOnlineAsyncEventSteamConnectionFailed(SteamSubsystem, RemoteId, (EP2PSessionError)CallbackData->m_eP2PSessionError);
+	UE_LOG_ONLINE(Verbose, TEXT("%s"), *NewEvent->ToString());
+	AddToOutQueue(NewEvent);
+}
+
+/**
+ * Notification event from Steam that a P2P connection has failed
+ */
+class FOnlineAsyncEventSteamShutdown : public FOnlineAsyncEvent<FOnlineSubsystemSteam>
+{
+	FOnlineAsyncEventSteamShutdown() :
+		FOnlineAsyncEvent(NULL)
+	{
+	}
+
+public:
+
+	FOnlineAsyncEventSteamShutdown(FOnlineSubsystemSteam* InSubsystem) :
+		FOnlineAsyncEvent(InSubsystem)
+	{
+	}
+
+	/**
+	 *	Get a human readable description of task
+	 */
+	virtual FString ToString() const override
+	{
+		return FString::Printf(TEXT("FOnlineAsyncEventSteamShutdown shutdown received."));	
+	}
+
+	/**
+	 * Give the async task a chance to marshal its data back to the game thread
+	 * Can only be called on the game thread by the async task manager
+	 */
+	virtual void Finalize() override
+	{
+		FPlatformMisc::RequestExit(false);
+	}
+};
+
+/**
+ *	Delegate registered with Steam to trigger when Steam is shutting down
+ *
+ * @param CallbackData - Steam struct containing shutdown information
+ */
+void FOnlineAsyncTaskManagerSteam::OnSteamShutdown(SteamShutdown_t* CallbackData)
+{
+	FOnlineAsyncEventSteamShutdown* NewEvent = new FOnlineAsyncEventSteamShutdown(SteamSubsystem);
+	UE_LOG_ONLINE(Verbose, TEXT("%s"), *NewEvent->ToString());
+	AddToOutQueue(NewEvent);
+}
diff --git a/Engine/Plugins/Online/OnlineSubsystemSteam/Source/Private/OnlineIdentityInterfaceSteam.h b/Engine/Plugins/Online/OnlineSubsystemSteam/Source/Private/OnlineIdentityInterfaceSteam.h
new file mode 100644
--- /dev/null
+++ b/Engine/Plugins/Online/OnlineSubsystemSteam/Source/Private/OnlineIdentityInterfaceSteam.h
@@ -0,0 +1,45 @@
+// Copyright 1998-2016 Epic Games, Inc. All Rights Reserved.
+
+#pragma once
+
+#include "OnlineSubsystemSteamTypes.h"
+#include "OnlineIdentityInterface.h"
+#include "OnlineSubsystemSteamPackage.h"
+
+class FOnlineIdentitySteam :
+	public IOnlineIdentity
+{
+	/** The steam user interface to use when interacting with steam */
+	class ISteamUser* SteamUserPtr;
+	/** The steam friends interface to use when interacting with steam */
+	class ISteamFriends* SteamFriendsPtr;
+
+PACKAGE_SCOPE:
+
+	FOnlineIdentitySteam();	
+
+public:
+
+	virtual ~FOnlineIdentitySteam() {};
+
+	// IOnlineIdentity
+
+	virtual bool Login(int32 LocalUserNum, const FOnlineAccountCredentials& AccountCredentials) override;
+	virtual bool Logout(int32 LocalUserNum) override;
+	virtual bool AutoLogin(int32 LocalUserNum) override;
+	virtual TSharedPtr<FUserOnlineAccount> GetUserAccount(const FUniqueNetId& UserId) const override;
+	virtual TArray<TSharedPtr<FUserOnlineAccount> > GetAllUserAccounts() const override;
+	virtual TSharedPtr<const FUniqueNetId> GetUniquePlayerId(int32 LocalUserNum) const override;
+	virtual TSharedPtr<const FUniqueNetId> CreateUniquePlayerId(uint8* Bytes, int32 Size) override;
+	virtual TSharedPtr<const FUniqueNetId> CreateUniquePlayerId(const FString& Str) override;
+	virtual ELoginStatus::Type GetLoginStatus(int32 LocalUserNum) const override;
+	virtual ELoginStatus::Type GetLoginStatus(const FUniqueNetId& UserId) const override;
+	virtual FString GetPlayerNickname(int32 LocalUserNum) const override;
+	virtual FString GetPlayerNickname(const FUniqueNetId& UserId) const override;
+	virtual FString GetAuthToken(int32 LocalUserNum) const override;
+	virtual void GetUserPrivilege(const FUniqueNetId& UserId, EUserPrivileges::Type Privilege, const FOnGetUserPrivilegeCompleteDelegate& Delegate) override;
+	virtual FPlatformUserId GetPlatformUserIdFromUniqueNetId(const FUniqueNetId& UniqueNetId) override;
+	virtual FString GetAuthType() const override;
+};
+
+typedef TSharedPtr<FOnlineIdentitySteam, ESPMode::ThreadSafe> FOnlineIdentitySteamPtr;
\ No newline at end of file
diff --git a/Engine/Plugins/Online/OnlineSubsystemSteam/Source/Private/OnlineSubsystemModuleSteam.cpp b/Engine/Plugins/Online/OnlineSubsystemSteam/Source/Private/OnlineSubsystemModuleSteam.cpp
new file mode 100644
--- /dev/null
+++ b/Engine/Plugins/Online/OnlineSubsystemSteam/Source/Private/OnlineSubsystemModuleSteam.cpp
@@ -0,0 +1,202 @@
+// Copyright 1998-2016 Epic Games, Inc. All Rights Reserved.
+
+#include "OnlineSubsystemSteamPrivatePCH.h"
+#include "OnlineSubsystemSteam.h"
+#include "ModuleManager.h"
+
+IMPLEMENT_MODULE(FOnlineSubsystemSteamModule, OnlineSubsystemSteam);
+
+//HACKTASTIC (Needed to keep delete function from being stripped out and crashing when protobuffers deallocate memory)
+void* HackDeleteFunctionPointer = (void*)(void(*)(void*))(::operator delete[]);
+
+/**
+ * Class responsible for creating instance(s) of the subsystem
+ */
+class FOnlineFactorySteam : public IOnlineFactory
+{
+
+private:
+
+	/** Single instantiation of the STEAM interface */
+	static FOnlineSubsystemSteamPtr SteamSingleton;
+
+	virtual void DestroySubsystem()
+	{
+		if (SteamSingleton.IsValid())
+		{
+			SteamSingleton->Shutdown();
+			SteamSingleton = NULL;
+		}
+	}
+
+public:
+
+	FOnlineFactorySteam() {}
+	virtual ~FOnlineFactorySteam() 
+	{
+		DestroySubsystem();
+	}
+
+	virtual IOnlineSubsystemPtr CreateSubsystem(FName InstanceName)
+	{
+		if (!SteamSingleton.IsValid())
+		{
+			SteamSingleton = MakeShareable(new FOnlineSubsystemSteam(InstanceName));
+			if (SteamSingleton->IsEnabled())
+			{
+				if(!SteamSingleton->Init())
+				{
+					UE_LOG_ONLINE(Warning, TEXT("Steam API failed to initialize!"));
+					DestroySubsystem();
+				}
+			}
+			else
+			{
+				UE_LOG_ONLINE(Warning, TEXT("Steam API disabled!"));
+				DestroySubsystem();
+			}
+
+			return SteamSingleton;
+		}
+
+		UE_LOG_ONLINE(Warning, TEXT("Can't create more than one instance of Steam online subsystem!"));
+		return NULL;
+	}
+};
+
+FOnlineSubsystemSteamPtr FOnlineFactorySteam::SteamSingleton = NULL;
+
+bool FOnlineSubsystemSteamModule::AreSteamDllsLoaded() const
+{
+	bool bLoadedClientDll = true;
+	bool bLoadedServerDll = true;
+
+#if LOADING_STEAM_LIBRARIES_DYNAMICALLY
+	bLoadedClientDll = (SteamDLLHandle != NULL) ? true : false;
+	#if LOADING_STEAM_SERVER_LIBRARY_DYNAMICALLY
+	bLoadedServerDll = IsRunningDedicatedServer() ? ((SteamServerDLLHandle != NULL) ? true : false) : true;
+	#endif //LOADING_STEAM_SERVER_LIBRARY_DYNAMICALLY
+#endif // LOADING_STEAM_LIBRARIES_DYNAMICALLY
+
+	return bLoadedClientDll && bLoadedServerDll;
+}
+
+static FString GetSteamModulePath()
+{
+#if PLATFORM_WINDOWS
+
+	#if PLATFORM_64BITS
+		return FPaths::EngineDir() / STEAM_SDK_ROOT_PATH / STEAM_SDK_VER_PATH / TEXT("Win64/");
+	#else
+		return FPaths::EngineDir() / STEAM_SDK_ROOT_PATH / STEAM_SDK_VER_PATH / TEXT("Win32/");
+	#endif	//PLATFORM_64BITS
+
+#else
+
+	return FString();
+
+#endif	//PLATFORM_WINDOWS
+}
+
+void FOnlineSubsystemSteamModule::LoadSteamModules()
+{
+	UE_LOG_ONLINE(Display, TEXT("Loading Steam SDK %s"), STEAM_SDK_VER);
+
+#if PLATFORM_WINDOWS
+	#if PLATFORM_64BITS
+		FString RootSteamPath = GetSteamModulePath();
+		FPlatformProcess::PushDllDirectory(*RootSteamPath);
+		SteamDLLHandle = FPlatformProcess::GetDllHandle(*(RootSteamPath + "steam_api64.dll"));
+#if 0 //64 bit not supported well at present, use Steam Client dlls
+		// [RCL] 2015-02-09 the above comment ("64 bit not supported well...") is from (early) 2012, things might have changed
+		// Load the Steam dedicated server dlls (assumes no Steam Client running)
+		if (IsRunningDedicatedServer())
+		{
+			SteamServerDLLHandle = FPlatformProcess::GetDllHandle(*(RootSteamPath + "steamclient64.dll"));
+		}
+#endif 
+		FPlatformProcess::PopDllDirectory(*RootSteamPath);
+	#else	//PLATFORM_64BITS
+		FString RootSteamPath = GetSteamModulePath();
+		FPlatformProcess::PushDllDirectory(*RootSteamPath);
+		SteamDLLHandle = FPlatformProcess::GetDllHandle(*(RootSteamPath + "steam_api.dll"));
+		if (IsRunningDedicatedServer())
+		{
+			SteamServerDLLHandle = FPlatformProcess::GetDllHandle(*(RootSteamPath + "steamclient.dll"));
+		}
+		FPlatformProcess::PopDllDirectory(*RootSteamPath);
+	#endif	//PLATFORM_64BITS
+#elif PLATFORM_MAC
+	SteamDLLHandle = FPlatformProcess::GetDllHandle(TEXT("libsteam_api.dylib"));
+#elif PLATFORM_LINUX
+
+#if LOADING_STEAM_LIBRARIES_DYNAMICALLY
+	UE_LOG_ONLINE(Log, TEXT("Loading system libsteam_api.so."));
+	SteamDLLHandle = FPlatformProcess::GetDllHandle(TEXT("libsteam_api.so"));
+	if (SteamDLLHandle == nullptr)
+	{
+		// try bundled one
+		UE_LOG_ONLINE(Log, TEXT("Could not find system one, loading bundled libsteam_api.so."));
+		FString RootSteamPath = FPaths::EngineDir() / FString::Printf(TEXT("Binaries/ThirdParty/Steamworks/%s/Linux/"), STEAM_SDK_VER_PATH);
+		SteamDLLHandle = FPlatformProcess::GetDllHandle(*(RootSteamPath + "libsteam_api.so"));
+	}
+#else
+	UE_LOG_ONLINE(Log, TEXT("libsteam_api.so is linked explicitly and should be already loaded."));
+#endif // LOADING_STEAM_LIBRARIES_DYNAMICALLY
+
+#endif	//PLATFORM_WINDOWS
+}
+
+void FOnlineSubsystemSteamModule::UnloadSteamModules()
+{
+#if LOADING_STEAM_LIBRARIES_DYNAMICALLY
+	if (SteamDLLHandle != NULL)
+	{
+		FPlatformProcess::FreeDllHandle(SteamDLLHandle);
+		SteamDLLHandle = NULL;
+	}
+
+	if (SteamServerDLLHandle != NULL)
+	{
+		FPlatformProcess::FreeDllHandle(SteamServerDLLHandle);
+		SteamServerDLLHandle = NULL;
+	}
+#endif	//LOADING_STEAM_LIBRARIES_DYNAMICALLY
+}
+
+void FOnlineSubsystemSteamModule::StartupModule()
+{
+	bool bSuccess = false;
+
+	// Load the Steam modules before first call to API
+	LoadSteamModules();
+	if (AreSteamDllsLoaded())
+	{
+		// Create and register our singleton factory with the main online subsystem for easy access
+		SteamFactory = new FOnlineFactorySteam();
+
+		FOnlineSubsystemModule& OSS = FModuleManager::GetModuleChecked<FOnlineSubsystemModule>("OnlineSubsystem");
+		OSS.RegisterPlatformService(STEAM_SUBSYSTEM, SteamFactory);
+		bSuccess = true;
+	}
+	else
+	{
+		UE_LOG_ONLINE(Warning, TEXT("Steam SDK %s libraries not present at %s or failed to load!"), STEAM_SDK_VER, *GetSteamModulePath());
+	}
+
+	if (!bSuccess)
+	{
+		UnloadSteamModules();
+	}
+}
+
+void FOnlineSubsystemSteamModule::ShutdownModule()
+{
+	FOnlineSubsystemModule& OSS = FModuleManager::GetModuleChecked<FOnlineSubsystemModule>("OnlineSubsystem");
+	OSS.UnregisterPlatformService(STEAM_SUBSYSTEM);
+	
+	delete SteamFactory;
+	SteamFactory = NULL;
+
+	UnloadSteamModules();
+}
diff --git a/Engine/Plugins/Online/OnlineSubsystemSteam/Source/Private/OnlineSubsystemSteamPrivatePCH.h b/Engine/Plugins/Online/OnlineSubsystemSteam/Source/Private/OnlineSubsystemSteamPrivatePCH.h
new file mode 100644
--- /dev/null
+++ b/Engine/Plugins/Online/OnlineSubsystemSteam/Source/Private/OnlineSubsystemSteamPrivatePCH.h
@@ -0,0 +1,57 @@
+// Copyright 1998-2016 Epic Games, Inc. All Rights Reserved.
+
+#pragma once
+
+#include "Engine.h"
+#include "OnlineSubsystemSteamModule.h"
+#include "OnlineSubsystemModule.h"
+#include "OnlineSubsystemUtils.h"
+#include "OnlineSubsystem.h"
+#include "OnlineSubsystemUtils.h"
+#include "SocketSubsystem.h"
+
+#include "ModuleManager.h"
+
+#define INVALID_INDEX -1
+
+/** Compile Steam SDK version in use */
+#define STEAM_SDK_VER TEXT("1.32")
+/** Path of the current Steam SDK version in use */
+#define STEAM_SDK_VER_PATH TEXT("Steamv132")
+/** Root location of Steam SDK */
+#define STEAM_SDK_ROOT_PATH TEXT("Binaries/ThirdParty/Steamworks")
+
+/** URL Prefix when using Steam socket connection */
+#define STEAM_URL_PREFIX TEXT("steam.")
+/** Filename containing the appid during development */
+#define STEAMAPPIDFILENAME TEXT("steam_appid.txt")
+
+/** pre-pended to all steam logging */
+#undef ONLINE_LOG_PREFIX
+#define ONLINE_LOG_PREFIX TEXT("STEAM: ")
+
+// @todo Steam: Steam headers trigger secure-C-runtime warnings in Visual C++. Rather than mess with _CRT_SECURE_NO_WARNINGS, we'll just
+//	disable the warnings locally. Remove when this is fixed in the SDK
+#ifdef _MSC_VER
+#pragma warning(push)
+#pragma warning(disable:4996)
+#endif
+
+#pragma push_macro("ARRAY_COUNT")
+#undef ARRAY_COUNT
+
+// Steamworks SDK headers
+#if STEAMSDK_FOUND == 0
+#error Steam SDK not located.  Expected to be found in Engine/Source/ThirdParty/Steamworks/{SteamVersion}
+#endif
+
+#include "steam/steam_api.h"
+#include "steam/steam_gameserver.h"
+
+#pragma pop_macro("ARRAY_COUNT")
+
+// @todo Steam: See above
+#ifdef _MSC_VER
+#pragma warning(pop)
+#endif
+
diff --git a/Engine/Plugins/Online/OnlineSubsystemSteam/Source/Private/OnlineUserCloudInterfaceSteam.cpp b/Engine/Plugins/Online/OnlineSubsystemSteam/Source/Private/OnlineUserCloudInterfaceSteam.cpp
new file mode 100644
--- /dev/null
+++ b/Engine/Plugins/Online/OnlineSubsystemSteam/Source/Private/OnlineUserCloudInterfaceSteam.cpp
@@ -0,0 +1,463 @@
+// Copyright 1998-2016 Epic Games, Inc. All Rights Reserved.
+
+#include "OnlineSubsystemSteamPrivatePCH.h"
+#include "OnlineSubsystemSteamTypes.h"
+#include "OnlineUserCloudInterfaceSteam.h"
+#include "OnlineSubsystemSteam.h"
+
+FString FOnlineAsyncTaskSteamEnumerateUserFiles::ToString() const
+{
+	return FString::Printf(TEXT("FOnlineAsyncTaskSteamEnumerateUserFiles bWasSuccessful:%d UserId:%s"),
+									bWasSuccessful, *UserId.ToDebugString());
+}
+
+void FOnlineAsyncTaskSteamEnumerateUserFiles::Tick()
+{
+	bIsComplete = true;
+	bWasSuccessful = false;
+
+	if (SteamRemoteStorage())
+	{
+		CSteamID SteamId(*(uint64*)UserId.GetBytes());
+		if (SteamUser()->BLoggedOn() && SteamUser()->GetSteamID() == SteamId)
+		{
+			//SteamSubsystem->GetUserCloudInterface()->DumpCloudState(UserId);
+
+			FScopeLock ScopeLock(&Subsystem->UserCloudDataLock);
+
+			// Get or create the user metadata entry and empty it
+			FSteamUserCloudData* UserMetadata = Subsystem->GetUserCloudEntry(UserId);
+
+			UserMetadata->CloudMetadata.Empty();
+
+			// Fill in the metadata entries
+			const int32 FileCount = (int32) SteamRemoteStorage()->GetFileCount();
+			for (int32 FileIdx = 0; FileIdx < FileCount; FileIdx++)
+			{
+				int32 FileSize = 0;
+				const char *FileName = SteamRemoteStorage()->GetFileNameAndSize(FileIdx, &FileSize);
+				new (UserMetadata->CloudMetadata) FCloudFileHeader(UTF8_TO_TCHAR(FileName), UTF8_TO_TCHAR(FileName), int32(FileSize));
+
+				//SteamSubsystem->GetUserCloudInterface()->DumpCloudFileState(UserId, FileName);
+			}
+
+			bWasSuccessful = true;
+		}
+		else
+		{
+			UE_LOG_ONLINE(Warning, TEXT("Can only enumerate cloud files for logged in user."));
+		}
+	}
+	else
+	{
+		UE_LOG_ONLINE(Warning, TEXT("Steam remote storage API disabled."));
+	}
+}
+
+void FOnlineAsyncTaskSteamEnumerateUserFiles::TriggerDelegates()
+{ 
+	FOnlineAsyncTaskSteam::TriggerDelegates();
+
+	IOnlineUserCloudPtr UserCloudInterface = Subsystem->GetUserCloudInterface();
+	UserCloudInterface->TriggerOnEnumerateUserFilesCompleteDelegates(bWasSuccessful, UserId);
+}
+
+FString FOnlineAsyncTaskSteamReadUserFile::ToString() const
+{
+	return FString::Printf(TEXT("FOnlineAsyncTaskSteamReadUserFile bWasSuccessful:%d UserId:%s FileName:%s"),
+									bWasSuccessful, *UserId.ToDebugString(), *FileName);
+}
+
+void FOnlineAsyncTaskSteamReadUserFile::Tick()
+{
+	// Going to be complete no matter what
+	bIsComplete = true;
+
+	if (SteamRemoteStorage() && FileName.Len() > 0)
+	{
+		CSteamID SteamId(*(uint64*)UserId.GetBytes());
+		if (SteamUser()->BLoggedOn() && SteamUser()->GetSteamID() == SteamId)
+		{
+			// Currently don't support greater than 1 chunk
+			const int32 FileSize = SteamRemoteStorage()->GetFileSize(TCHAR_TO_UTF8(*FileName));
+			if (FileSize >= 0 && FileSize <= k_unMaxCloudFileChunkSize)
+			{
+				FScopeLock ScopeLock(&Subsystem->UserCloudDataLock);
+				// Create or get the current entry for this file
+				FSteamUserCloudData* UserCloud = Subsystem->GetUserCloudEntry(UserId);
+				if (UserCloud)
+				{
+					FCloudFile* UserCloudFile = UserCloud->GetFileData(FileName, true);
+					check(UserCloudFile);
+
+					// Allocate and read in the file
+					UserCloudFile->Data.Empty(FileSize);
+					UserCloudFile->Data.AddUninitialized(FileSize);
+					if (SteamRemoteStorage()->FileRead(TCHAR_TO_UTF8(*FileName), UserCloudFile->Data.GetData(), FileSize) == FileSize)
+					{
+						bWasSuccessful = true;
+					}
+					else
+					{
+						UserCloudFile->Data.Empty();
+					}
+				}
+			}
+			else
+			{
+				UE_LOG_ONLINE(Warning, TEXT("Requested file %s has invalid size %d."), *FileName, FileSize);
+			}
+		}	
+		else
+		{
+			UE_LOG_ONLINE(Warning, TEXT("Can only read cloud files for logged in user."));
+		}
+	}
+	else
+	{
+		UE_LOG_ONLINE(Warning, TEXT("Steam remote storage API disabled."));
+	}
+
+	{
+		FScopeLock ScopeLock(&Subsystem->UserCloudDataLock);
+		FSteamUserCloudData* UserCloud = Subsystem->GetUserCloudEntry(UserId);
+		if (UserCloud)
+		{
+			FCloudFile* UserCloudFileData = UserCloud->GetFileData(FileName);
+			if (UserCloudFileData)
+			{
+				UserCloudFileData->AsyncState = bWasSuccessful ? EOnlineAsyncTaskState::Done : EOnlineAsyncTaskState::Failed;
+			}
+		}
+	}
+}
+
+void FOnlineAsyncTaskSteamReadUserFile::TriggerDelegates()
+{ 
+	FOnlineAsyncTaskSteam::TriggerDelegates();
+
+	IOnlineUserCloudPtr UserCloudInterface = Subsystem->GetUserCloudInterface();
+	UserCloudInterface->TriggerOnReadUserFileCompleteDelegates(bWasSuccessful, UserId, FileName);
+}
+
+bool FOnlineAsyncTaskSteamWriteUserFile::WriteUserFile(const FUniqueNetId& InUserId, const FString& InFileToWrite, const TArray<uint8>& InContents)
+{
+	bool bSuccess = false;
+	if (InFileToWrite.Len() > 0 && InContents.Num() > 0)
+	{
+		if (SteamRemoteStorage() && FileName.Len() > 0)
+		{
+			CSteamID SteamId(*(uint64*)InUserId.GetBytes());
+			if (SteamUser()->BLoggedOn() && SteamUser()->GetSteamID() == SteamId)
+			{
+				// Currently don't support greater than 1 chunk
+				if (InContents.Num() < k_unMaxCloudFileChunkSize)
+				{
+					if (SteamRemoteStorage()->FileWrite(TCHAR_TO_UTF8(*InFileToWrite), InContents.GetData(), InContents.Num()))
+					{
+						FScopeLock ScopeLock(&Subsystem->UserCloudDataLock);
+						FSteamUserCloudData* UserCloud = Subsystem->GetUserCloudEntry(InUserId);
+						if (UserCloud)
+						{
+							// Update the metadata table to reflect this write (might be new entry)
+							FCloudFileHeader* UserCloudFileMetadata = UserCloud->GetFileMetadata(InFileToWrite, true);
+							check(UserCloudFileMetadata);
+
+							UserCloudFileMetadata->FileSize = SteamRemoteStorage()->GetFileSize(TCHAR_TO_UTF8(*InFileToWrite));
+							UserCloudFileMetadata->Hash = FString(TEXT("0"));
+
+							// Update the file table to reflect this write
+							FCloudFile* UserCloudFileData = UserCloud->GetFileData(InFileToWrite, true);
+							check(UserCloudFileData);
+
+							UserCloudFileData->Data = InContents;
+							bSuccess = true;
+						}
+					}
+					else
+					{
+						UE_LOG_ONLINE(Warning, TEXT("Failed to write file to Steam cloud \"%s\"."), *InFileToWrite);
+					}
+				}
+				else
+				{
+					UE_LOG_ONLINE(Warning, TEXT("File too large %d to write to Steam cloud."), InContents.Num());
+				}
+			}
+			else
+			{
+				UE_LOG_ONLINE(Warning, TEXT("Can only write cloud files for logged in user."));
+			}
+		}
+		else
+		{
+			UE_LOG_ONLINE(Warning, TEXT("Steam remote storage API disabled."));
+		}
+	}
+
+	{
+		FScopeLock ScopeLock(&Subsystem->UserCloudDataLock);
+		FSteamUserCloudData* UserCloud = Subsystem->GetUserCloudEntry(InUserId);
+		if (UserCloud)
+		{
+			FCloudFile* UserCloudFileData = UserCloud->GetFileData(InFileToWrite, true);
+			check(UserCloudFileData);
+			UserCloudFileData->AsyncState = bSuccess ? EOnlineAsyncTaskState::Done : EOnlineAsyncTaskState::Failed;
+		}
+	}
+
+	//SteamSubsystem->GetUserCloudInterface()->DumpCloudFileState(InUserId, InFileToWrite);
+	return bSuccess;
+}
+
+FString FOnlineAsyncTaskSteamWriteUserFile::ToString() const
+{
+	return FString::Printf(TEXT("FOnlineAsyncTaskSteamWriteUserFile bWasSuccessful:%d UserId:%s FileName:%s"),
+									bWasSuccessful, *UserId.ToDebugString(), *FileName);
+}
+
+void FOnlineAsyncTaskSteamWriteUserFile::Tick()
+{	
+	// Going to be complete no matter what
+	bIsComplete = true;
+	if (WriteUserFile(UserId, FileName, Contents))
+	{
+		bWasSuccessful = true;
+	}
+
+	// Done with this copy of the data regardless
+	Contents.Empty();
+}
+
+void FOnlineAsyncTaskSteamWriteUserFile::TriggerDelegates()
+{ 
+	FOnlineAsyncTaskSteam::TriggerDelegates();
+
+	IOnlineUserCloudPtr UserCloudInterface = Subsystem->GetUserCloudInterface();
+	UserCloudInterface->TriggerOnWriteUserFileCompleteDelegates(bWasSuccessful, UserId, FileName);
+}
+
+FOnlineUserCloudSteam::~FOnlineUserCloudSteam()
+{
+	if (SteamSubsystem)
+	{
+		SteamSubsystem->ClearUserCloudFiles();
+	}
+}
+
+bool FOnlineUserCloudSteam::GetFileContents(const FUniqueNetId& UserId, const FString& FileName, TArray<uint8>& FileContents)
+{
+	FScopeLock(&SteamSubsystem->UserCloudDataLock);
+	// Search for the specified file and return the raw data
+	FSteamUserCloudData* UserCloudData = SteamSubsystem->GetUserCloudEntry(UserId);
+	if (UserCloudData)
+	{
+		FCloudFile* SteamCloudFile = UserCloudData->GetFileData(FileName);
+		if (SteamCloudFile && SteamCloudFile->AsyncState == EOnlineAsyncTaskState::Done && SteamCloudFile->Data.Num() > 0)
+		{
+			FileContents = SteamCloudFile->Data;
+			return true;
+		}
+	}
+	
+	return false;
+}
+
+bool FOnlineUserCloudSteam::ClearFiles(const FUniqueNetId& UserId)
+{
+	FScopeLock(&SteamSubsystem->UserCloudDataLock);
+	// Search for the specified file
+	FSteamUserCloudData* UserCloudData = SteamSubsystem->GetUserCloudEntry(UserId);
+	if (UserCloudData)
+	{
+		return UserCloudData->ClearFiles();
+	}
+
+	return true;
+}
+
+bool FOnlineUserCloudSteam::ClearFile(const FUniqueNetId& UserId, const FString& FileName)
+{
+	FScopeLock(&SteamSubsystem->UserCloudDataLock);
+	// Search for the specified file
+	FSteamUserCloudData* UserCloudData = SteamSubsystem->GetUserCloudEntry(UserId);
+	if (UserCloudData)
+	{
+		return UserCloudData->ClearFileData(FileName);
+	}
+
+	return true;
+}
+
+void FOnlineUserCloudSteam::EnumerateUserFiles(const FUniqueNetId& UserId)
+{
+	SteamSubsystem->QueueAsyncTask(new FOnlineAsyncTaskSteamEnumerateUserFiles(SteamSubsystem, FUniqueNetIdSteam(*(uint64*)UserId.GetBytes())));
+}
+
+void FOnlineUserCloudSteam::GetUserFileList(const FUniqueNetId& UserId, TArray<FCloudFileHeader>& UserFiles)
+{
+	FScopeLock(&SteamSubsystem->UserCloudDataLock);
+	FSteamUserCloudData* UserMetadata = SteamSubsystem->GetUserCloudEntry(UserId);  
+	UserFiles = UserMetadata->CloudMetadata;
+}
+
+bool FOnlineUserCloudSteam::ReadUserFile(const FUniqueNetId& UserId, const FString& FileName)
+{
+	FScopeLock(&SteamSubsystem->UserCloudDataLock);
+	// Create or get the current entry for this file
+	FSteamUserCloudData* UserCloud = SteamSubsystem->GetUserCloudEntry(UserId);
+	if (UserCloud && FileName.Len() > 0)
+	{
+		FCloudFile* UserCloudFile = UserCloud->GetFileData(FileName, true);
+		UserCloudFile->AsyncState = EOnlineAsyncTaskState::InProgress;
+		SteamSubsystem->QueueAsyncTask(new FOnlineAsyncTaskSteamReadUserFile(SteamSubsystem, FUniqueNetIdSteam(*(uint64*)UserId.GetBytes()), FileName));
+		return true;
+	}
+	
+	return false;
+}
+
+bool FOnlineUserCloudSteam::WriteUserFile(const FUniqueNetId& UserId, const FString& FileName, TArray<uint8>& FileContents)
+{
+	FScopeLock(&SteamSubsystem->UserCloudDataLock);
+	// Create or get the current entry for this file
+	FSteamUserCloudData* UserCloud = SteamSubsystem->GetUserCloudEntry(UserId);
+	if (UserCloud && FileName.Len() > 0)
+	{
+		FCloudFile* UserCloudFile = UserCloud->GetFileData(FileName, true);
+		UserCloudFile->AsyncState = EOnlineAsyncTaskState::InProgress;
+		SteamSubsystem->QueueAsyncTask(new FOnlineAsyncTaskSteamWriteUserFile(SteamSubsystem, FUniqueNetIdSteam(*(uint64*)UserId.GetBytes()), FileName, FileContents));
+		return true;
+	}
+
+	return false;
+}
+
+void FOnlineUserCloudSteam::CancelWriteUserFile(const FUniqueNetId& UserId, const FString& FileName)
+{
+	// Not implemented
+}
+
+
+bool FOnlineUserCloudSteam::DeleteUserFile(const FUniqueNetId& UserId, const FString& FileName, bool bShouldCloudDelete, bool bShouldLocallyDelete)
+{
+	SteamSubsystem->QueueAsyncTask(new FOnlineAsyncTaskSteamDeleteUserFile(SteamSubsystem, FUniqueNetIdSteam(*(uint64*)UserId.GetBytes()), FileName, bShouldCloudDelete, bShouldLocallyDelete));
+	return true;
+}
+
+bool FOnlineUserCloudSteam::RequestUsageInfo(const FUniqueNetId& UserId)
+{
+	// Not implemented
+	return false;
+}
+
+void FOnlineUserCloudSteam::DumpCloudState(const FUniqueNetId& UserId)
+{
+	int32 TotalBytes, TotalAvailable;
+	if (SteamRemoteStorage()->GetQuota(&TotalBytes, &TotalAvailable) == false)
+	{
+		TotalBytes = TotalAvailable = 0;
+	}
+
+	UE_LOG_ONLINE(Verbose, TEXT("Steam Disk Quota: %d / %d"), TotalAvailable, TotalBytes);
+	UE_LOG_ONLINE(Verbose, TEXT("Game does %shave cloud storage enabled."), SteamRemoteStorage()->IsCloudEnabledForApp() ? TEXT("") : TEXT("NOT "));
+	UE_LOG_ONLINE(Verbose, TEXT("User does %shave cloud storage enabled."), SteamRemoteStorage()->IsCloudEnabledForAccount() ? TEXT("") : TEXT("NOT "));
+}
+
+void FOnlineUserCloudSteam::DumpCloudFileState(const FUniqueNetId& UserId, const FString& FileName)
+{
+	if (FileName.Len() > 0)
+	{
+		UE_LOG_ONLINE(Log, TEXT("Cloud File State file %s:"), *FileName);
+		{
+			FScopeLock(&SteamSubsystem->UserCloudDataLock);
+			FSteamUserCloudData* UserCloud = SteamSubsystem->GetUserCloudEntry(UserId);
+			FCloudFileHeader* FileMetadata = UserCloud->GetFileMetadata(FileName);
+			if (FileMetadata)
+			{
+				UE_LOG_ONLINE(Log, TEXT("\tMeta: FileName:%s DLName:%s FileSize:%d Hash:%s"), *FileMetadata->FileName, *FileMetadata->DLName, FileMetadata->FileSize, *FileMetadata->Hash);
+			}
+			else
+			{
+				UE_LOG_ONLINE(Log, TEXT("\tNo metadata found!"));
+			}
+
+			FCloudFile* FileData = UserCloud->GetFileData(FileName);
+			if (FileData)
+			{
+				UE_LOG_ONLINE(Log, TEXT("\tFileCache: FileName:%s State:%s CacheSize:%d"), *FileData->FileName, EOnlineAsyncTaskState::ToString(FileData->AsyncState), FileData->Data.Num());
+			}
+			else
+			{
+				UE_LOG_ONLINE(Log, TEXT("\tNo cache entry found!"));
+			}
+		}
+
+		int32 FileSize = SteamRemoteStorage()->GetFileSize(TCHAR_TO_UTF8(*FileName));
+
+		UE_LOG_ONLINE(Log, TEXT("\tSteam: FileName:%s Size:%d Exists:%s Persistent:%s"),
+			*FileName, FileSize, 
+			SteamRemoteStorage()->FileExists(TCHAR_TO_UTF8(*FileName)) ? TEXT("Y") : TEXT("N"),
+			SteamRemoteStorage()->FilePersisted(TCHAR_TO_UTF8(*FileName)) ? TEXT("Y") : TEXT("N"));
+	}
+}
+
+FString FOnlineAsyncTaskSteamDeleteUserFile::ToString() const
+{
+	return FString::Printf(TEXT("FOnlineAsyncTaskSteamDeleteUserFile bWasSuccessful:%d UserId:%s FileName:%s"),
+								bWasSuccessful, *UserId.ToDebugString(), *FileName);
+}
+
+void FOnlineAsyncTaskSteamDeleteUserFile::Tick()
+{
+	bWasSuccessful = false;
+
+	if (SteamRemoteStorage() && FileName.Len() > 0)
+	{
+		CSteamID SteamId(*(uint64*)UserId.GetBytes());
+		if (SteamUser()->BLoggedOn() && SteamUser()->GetSteamID() == SteamId)
+		{
+			bool bCloudDeleteSuccess = true;
+			if (bShouldCloudDelete)
+			{
+				// Remove the cloud flag, the file remains safely available on the local machine
+				bCloudDeleteSuccess = SteamRemoteStorage()->FileForget(TCHAR_TO_UTF8(*FileName));
+			}
+
+			bool bLocalDeleteSuccess = true;
+			if (bShouldLocallyDelete)
+			{
+				bLocalDeleteSuccess = false;
+				// Only clear the tables if we're permanently deleting the file
+				// Need to make sure nothing async is happening	first (this is a formality as nothing in Steam actually is)
+				IOnlineUserCloudPtr UserCloud = Subsystem->GetUserCloudInterface();
+				if (UserCloud->ClearFile(UserId, FileName))
+				{
+					// Permanent delete
+					bLocalDeleteSuccess = SteamRemoteStorage()->FileDelete(TCHAR_TO_UTF8(*FileName));
+					Subsystem->ClearUserCloudMetadata(UserId, FileName);
+				}
+			}
+
+			bWasSuccessful = bCloudDeleteSuccess && bLocalDeleteSuccess;
+		}	
+		else
+		{
+			UE_LOG_ONLINE(Warning, TEXT("Can only delete cloud files for logged in user."));
+		}
+	}
+	else
+	{
+		UE_LOG_ONLINE(Warning, TEXT("Steam remote storage API disabled."));
+	}
+
+	bIsComplete = true;
+}
+
+void FOnlineAsyncTaskSteamDeleteUserFile::TriggerDelegates()
+{ 
+	FOnlineAsyncTaskSteam::TriggerDelegates();
+
+	IOnlineUserCloudPtr UserCloudInterface = Subsystem->GetUserCloudInterface();
+	UserCloudInterface->TriggerOnDeleteUserFileCompleteDelegates(bWasSuccessful, UserId, FileName);
+}
