Copying //Tasks/UE4/Dev-Online-Plugins-v3 to //UE4/Dev-Main (Source: //Tasks/UE4/Dev-Online-Plugins-v3 @ 3047294)

diff --git a/Engine/Plugins/Online/OnlineSubsystemSteam/Source/Private/SocketSubsystemSteam.cpp b/Engine/Plugins/Online/OnlineSubsystemSteam/Source/Private/SocketSubsystemSteam.cpp
new file mode 100644
--- /dev/null
+++ b/Engine/Plugins/Online/OnlineSubsystemSteam/Source/Private/SocketSubsystemSteam.cpp
@@ -0,0 +1,584 @@
+// Copyright 1998-2016 Epic Games, Inc. All Rights Reserved.
+
+#include "OnlineSubsystemSteamPrivatePCH.h"
+#include "SocketSubsystemSteam.h"
+#include "SocketsSteam.h"
+#include "IPAddressSteam.h"
+#include "OnlineSessionInterfaceSteam.h"
+#include "SocketSubsystemModule.h"
+#include "ModuleManager.h"
+#include "Classes/SteamNetConnection.h"
+#include "Engine.h"
+
+FSocketSubsystemSteam* FSocketSubsystemSteam::SocketSingleton = NULL;
+
+/**
+ * Create the socket subsystem for the given platform service
+ */
+FName CreateSteamSocketSubsystem()
+{
+	// Create and register our singleton factory with the main online subsystem for easy access
+	FSocketSubsystemSteam* SocketSubsystem = FSocketSubsystemSteam::Create();
+	FString Error;
+	if (SocketSubsystem->Init(Error))
+	{
+		FSocketSubsystemModule& SSS = FModuleManager::LoadModuleChecked<FSocketSubsystemModule>("Sockets");
+		SSS.RegisterSocketSubsystem(STEAM_SUBSYSTEM, SocketSubsystem, true);
+		return STEAM_SUBSYSTEM;
+	}
+	else
+	{
+		FSocketSubsystemSteam::Destroy();
+		return NAME_None;
+	}
+}
+
+/**
+ * Tear down the socket subsystem for the given platform service
+ */
+void DestroySteamSocketSubsystem()
+{
+	FModuleManager& ModuleManager = FModuleManager::Get();
+
+	if (ModuleManager.IsModuleLoaded("Sockets"))
+	{
+		FSocketSubsystemModule& SSS = FModuleManager::GetModuleChecked<FSocketSubsystemModule>("Sockets");
+		SSS.UnregisterSocketSubsystem(STEAM_SUBSYSTEM);
+	}
+	FSocketSubsystemSteam::Destroy();
+}
+
+/** 
+ * Singleton interface for this subsystem 
+ * @return the only instance of this subsystem
+ */
+FSocketSubsystemSteam* FSocketSubsystemSteam::Create()
+{
+	if (SocketSingleton == NULL)
+	{
+		SocketSingleton = new FSocketSubsystemSteam();
+	}
+
+	return SocketSingleton;
+}
+
+/**
+ * Performs Steam specific socket clean up
+ */
+void FSocketSubsystemSteam::Destroy()
+{
+	if (SocketSingleton != NULL)
+	{
+		SocketSingleton->Shutdown();
+		delete SocketSingleton;
+		SocketSingleton = NULL;
+	}
+}
+
+/**
+ * Does Steam platform initialization of the sockets library
+ *
+ * @param Error a string that is filled with error information
+ *
+ * @return true if initialized ok, false otherwise
+ */
+bool FSocketSubsystemSteam::Init(FString& Error)
+{
+	if (GConfig)
+	{
+		if (!GConfig->GetBool(TEXT("OnlineSubsystemSteam"), TEXT("bAllowP2PPacketRelay"), bAllowP2PPacketRelay, GEngineIni))
+		{
+			UE_LOG_ONLINE(Warning, TEXT("Missing bAllowP2PPacketRelay key in OnlineSubsystemSteam of DefaultEngine.ini"));
+		}
+
+		if (!GConfig->GetFloat(TEXT("OnlineSubsystemSteam"), TEXT("P2PConnectionTimeout"), P2PConnectionTimeout, GEngineIni))
+		{
+			UE_LOG_ONLINE(Warning, TEXT("Missing P2PConnectionTimeout key in OnlineSubsystemSteam of DefaultEngine.ini"));
+		}
+	}
+
+	if (SteamNetworking())
+	{
+		SteamNetworking()->AllowP2PPacketRelay(bAllowP2PPacketRelay);
+	}
+
+	if (SteamGameServerNetworking())
+	{
+		SteamGameServerNetworking()->AllowP2PPacketRelay(bAllowP2PPacketRelay);
+	}
+
+	return true;
+}
+
+/**
+ * Performs platform specific socket clean up
+ */
+void FSocketSubsystemSteam::Shutdown()
+{
+	for (int32 ConnIdx=SteamConnections.Num()-1; ConnIdx>=0; ConnIdx--)
+	{
+		if (SteamConnections[ConnIdx].IsValid())
+		{
+			USteamNetConnection* SteamConn = CastChecked<USteamNetConnection>(SteamConnections[ConnIdx].Get());
+			UnregisterConnection(SteamConn);
+		}
+	}
+
+	// Cleanup any remaining sessions
+	TArray<FUniqueNetIdSteam> SessionIds;
+	AcceptedConnections.GenerateKeyArray(SessionIds);
+	for (int SessionIdx=0; SessionIdx < SessionIds.Num(); SessionIdx++)
+	{
+		P2PRemove(SessionIds[SessionIdx]);
+	}
+	
+	CleanupDeadConnections();
+
+	// Cleanup sockets
+	TArray<FSocketSteam*> TempArray = SteamSockets;
+	for (int SocketIdx=0; SocketIdx < TempArray.Num(); SocketIdx++)
+	{
+		DestroySocket(TempArray[SocketIdx]);
+	}
+
+	SteamSockets.Empty();
+	SteamConnections.Empty();
+	AcceptedConnections.Empty();
+	DeadConnections.Empty();
+}
+
+/**
+ * Creates a socket
+ *
+ * @Param SocketType type of socket to create (DGram, Stream, etc)
+ * @param SocketDescription debug description
+ * @param bForceUDP overrides any platform specific protocol with UDP instead
+ *
+ * @return the new socket or NULL if failed
+ */
+FSocket* FSocketSubsystemSteam::CreateSocket(const FName& SocketType, const FString& SocketDescription, bool bForceUDP)
+{
+	FSocket* NewSocket = NULL;
+	if (SocketType == FName("SteamClientSocket"))
+	{
+		ISteamUser* SteamUserPtr = SteamUser();
+		if (SteamUserPtr != NULL)
+		{
+			FUniqueNetIdSteam ClientId(SteamUserPtr->GetSteamID());
+			NewSocket = new FSocketSteam(SteamNetworking(), ClientId, SocketDescription);
+
+			if (NewSocket)
+			{
+				AddSocket((FSocketSteam*)NewSocket);
+			}
+		}
+	}
+	else if (SocketType == FName("SteamServerSocket"))
+	{
+		IOnlineSubsystem* SteamSubsystem = IOnlineSubsystem::Get(STEAM_SUBSYSTEM);
+		FOnlineSessionSteamPtr SessionInt = StaticCastSharedPtr<FOnlineSessionSteam>(SteamSubsystem->GetSessionInterface());
+		if (SessionInt.IsValid())
+		{
+			// If the GameServer connection hasn't been created yet, mark the socket as invalid for now
+			if (SessionInt->bSteamworksGameServerConnected && SessionInt->GameServerSteamId->IsValid() && SessionInt->bPolicyResponseReceived)
+			{
+				NewSocket = new FSocketSteam(SteamGameServerNetworking(), *SessionInt->GameServerSteamId, SocketDescription);
+			}
+			else
+			{
+				FUniqueNetIdSteam InvalidId(uint64(0));
+				NewSocket = new FSocketSteam(SteamGameServerNetworking(), InvalidId, SocketDescription);
+			}
+
+			if (NewSocket)
+			{
+				AddSocket((FSocketSteam*)NewSocket);
+			}
+		}
+	}
+	else
+	{
+		ISocketSubsystem* PlatformSocketSub = ISocketSubsystem::Get(PLATFORM_SOCKETSUBSYSTEM);
+		if (PlatformSocketSub)
+		{
+			NewSocket = PlatformSocketSub->CreateSocket(SocketType, SocketDescription, bForceUDP);
+		}
+	}
+	
+	if (!NewSocket)
+	{
+		UE_LOG(LogSockets, Warning, TEXT("Failed to create socket %s [%s]"), *SocketType.ToString(), *SocketDescription);
+	}
+
+	return NewSocket;
+}
+
+/**
+ * Cleans up a socket class
+ *
+ * @param Socket the socket object to destroy
+ */
+void FSocketSubsystemSteam::DestroySocket(FSocket* Socket)
+{
+	// Possible non steam socket here PLATFORM_SOCKETSUBSYSTEM, but its just a pointer compare
+	RemoveSocket((FSocketSteam*)Socket);
+	delete Socket;
+}
+
+/**
+ * Associate the game server steam id with any sockets that were created prior to successful login
+ *
+ * @param GameServerId id assigned to this game server
+ */
+void FSocketSubsystemSteam::FixupSockets(const FUniqueNetIdSteam& GameServerId)
+{
+	for (int32 SockIdx = 0; SockIdx < SteamSockets.Num(); SockIdx++)
+	{
+		FSocketSteam* Socket = SteamSockets[SockIdx];
+		if (Socket->SteamNetworkingPtr == SteamGameServerNetworking() && !Socket->LocalSteamId.IsValid())
+		{
+			Socket->LocalSteamId = GameServerId;
+		}
+	}
+}
+
+/**
+ * Adds a steam connection for tracking
+ *
+ * @param Connection The connection to add for tracking
+ */
+void FSocketSubsystemSteam::RegisterConnection(USteamNetConnection* Connection)
+{
+	check(!Connection->bIsPassthrough);
+
+	FWeakObjectPtr ObjectPtr = Connection;
+	SteamConnections.Add(ObjectPtr);
+
+	if (Connection->RemoteAddr.IsValid() && Connection->Socket)
+	{
+		FSocketSteam* SteamSocket = (FSocketSteam*)Connection->Socket;
+		FInternetAddrSteam& SteamAddr = (FInternetAddrSteam&)(*Connection->RemoteAddr);
+		P2PTouch(SteamSocket->SteamNetworkingPtr, SteamAddr.SteamId);
+	}
+}
+
+/**
+ * Removes a steam connection from tracking
+ *
+ * @param Connection The connection to remove from tracking
+ */
+void FSocketSubsystemSteam::UnregisterConnection(USteamNetConnection* Connection)
+{
+	check(!Connection->bIsPassthrough);
+
+	FWeakObjectPtr ObjectPtr = Connection;
+	int32 NumRemoved = SteamConnections.RemoveSingleSwap(ObjectPtr);
+
+	// Don't call P2PRemove again if we didn't actually remove a connection. This 
+	// will get called twice - once the connection is closed and when the connection
+	// is garbage collected. It's possible that the player who left rejoined before garbage
+	// collection runs (their connection object will be different), so P2PRemove would kick
+	// them from the session when it shouldn't.
+	if (NumRemoved > 0 && Connection->RemoteAddr.IsValid())
+	{
+		FInternetAddrSteam& SteamAddr = (FInternetAddrSteam&)(*Connection->RemoteAddr);
+		P2PRemove(SteamAddr.SteamId, SteamAddr.SteamChannel);
+	}
+}
+
+void FSocketSubsystemSteam::ConnectFailure(FUniqueNetIdSteam& RemoteId)
+{
+	// Remove any GC'd references
+	for (int32 ConnIdx=SteamConnections.Num()-1; ConnIdx>=0; ConnIdx--)
+	{
+		if (!SteamConnections[ConnIdx].IsValid())
+		{
+			SteamConnections.RemoveAt(ConnIdx);
+		}
+	}
+
+	// Find the relevant connections and shut them down
+	for (int32 ConnIdx=0; ConnIdx<SteamConnections.Num(); ConnIdx++)
+	{
+		USteamNetConnection* SteamConn = CastChecked<USteamNetConnection>(SteamConnections[ConnIdx].Get());
+		FInternetAddrSteam& RemoteAddrSteam = (FInternetAddrSteam&)*(SteamConn->RemoteAddr);
+
+		// Only checking Id here because its a complete failure (channel doesn't matter)
+		if (RemoteAddrSteam.SteamId == RemoteId)
+		{
+			SteamConn->Close();
+		}
+	}
+
+	P2PRemove(RemoteId);
+}
+
+/**
+ * Does a DNS look up of a host name
+ *
+ * @param HostName the name of the host to look up
+ * @param OutAddr the address to copy the IP address to
+ */
+ESocketErrors FSocketSubsystemSteam::GetHostByName(const ANSICHAR* HostName, FInternetAddr& OutAddr)
+{
+	return SE_EADDRNOTAVAIL;
+}
+
+/**
+ * Determines the name of the local machine
+ *
+ * @param HostName the string that receives the data
+ *
+ * @return true if successful, false otherwise
+ */
+bool FSocketSubsystemSteam::GetHostName(FString& HostName)
+{
+	return false;
+}
+
+/**
+ *	Create a proper FInternetAddr representation
+ * @param Address host address
+ * @param Port host port
+ */
+TSharedRef<FInternetAddr> FSocketSubsystemSteam::CreateInternetAddr(uint32 Address, uint32 Port)
+{
+	FInternetAddrSteam* SteamAddr = new FInternetAddrSteam();
+	return MakeShareable(SteamAddr);
+}
+
+/**
+ * @return Whether the machine has a properly configured network device or not
+ */
+bool FSocketSubsystemSteam::HasNetworkDevice() 
+{
+	return true;
+}
+
+/**
+ *	Get the name of the socket subsystem
+ * @return a string naming this subsystem
+ */
+const TCHAR* FSocketSubsystemSteam::GetSocketAPIName() const 
+{
+	return TEXT("SteamSockets");
+}
+
+/**
+ * Returns the last error that has happened
+ */
+ESocketErrors FSocketSubsystemSteam::GetLastErrorCode()
+{
+	return TranslateErrorCode(LastSocketError);
+}
+
+/**
+ * Translates the platform error code to a ESocketErrors enum
+ */
+ESocketErrors FSocketSubsystemSteam::TranslateErrorCode(int32 Code)
+{
+	// @TODO ONLINE - This needs to be filled in (at present it is 1:1)
+	return (ESocketErrors)LastSocketError;
+}
+
+/**
+ *	Get local IP to bind to
+ */
+TSharedRef<FInternetAddr> FSocketSubsystemSteam::GetLocalBindAddr(FOutputDevice& Out)
+{
+	FInternetAddrSteam* SteamAddr = NULL;
+	CSteamID SteamId;
+	if (SteamUser())
+	{
+		// Prefer the steam user
+		SteamId = SteamUser()->GetSteamID();
+		SteamAddr = new FInternetAddrSteam(FUniqueNetIdSteam(SteamId));
+	}
+	else if (SteamGameServer() && SteamGameServer()->BLoggedOn())
+	{
+		// Dedicated server 
+		SteamId = SteamGameServer()->GetSteamID();
+		SteamAddr = new FInternetAddrSteam(FUniqueNetIdSteam(SteamId));
+	}
+	else
+	{
+		// Empty/invalid case
+		SteamAddr = new FInternetAddrSteam();
+	}
+
+	return MakeShareable(SteamAddr);
+}
+
+/**
+ * Potentially accept an incoming connection from a Steam P2P request
+ * 
+ * @param SteamNetworkingPtr the interface for access the P2P calls (Client/GameServer)
+ * @param RemoteId the id of the incoming request
+ * 
+ * @return true if accepted, false otherwise
+ */
+bool FSocketSubsystemSteam::AcceptP2PConnection(ISteamNetworking* SteamNetworkingPtr, FUniqueNetIdSteam& RemoteId)
+{
+	if (SteamNetworkingPtr && RemoteId.IsValid() && DeadConnections.Find(RemoteId) == NULL)
+	{
+		// Blindly accept connections (but only if P2P enabled)
+		SteamNetworkingPtr->AcceptP2PSessionWithUser(RemoteId);
+		AcceptedConnections.Add(RemoteId, FSteamP2PConnectionInfo(SteamNetworkingPtr, FPlatformTime::Seconds()));
+		return true;
+	}
+
+	return false;
+}
+
+/**
+ * Add/update a Steam P2P connection as being recently accessed
+ *
+ * @param SteamNetworkingPtr proper networking interface that this session is communicating on
+ * @param SessionId P2P session recently heard from
+ *
+ * @return true if the connection is active, false if this is in the dead connections list
+ */
+bool FSocketSubsystemSteam::P2PTouch(ISteamNetworking* SteamNetworkingPtr, FUniqueNetIdSteam& SessionId)
+{
+    // Don't update any sessions coming from pending disconnects
+	if (DeadConnections.Find(SessionId) == NULL)
+	{
+		AcceptedConnections.Add(SessionId, FSteamP2PConnectionInfo(SteamNetworkingPtr, FPlatformTime::Seconds()));
+		return true;
+	}
+
+	return false;
+}
+	
+/**
+ * Remove a Steam P2P session from tracking and close the connection
+ *
+ * @param SessionId P2P session to remove
+ * @param Channel channel to close, -1 to close all communication
+ */
+void FSocketSubsystemSteam::P2PRemove(FUniqueNetIdSteam& SessionId, int32 Channel)
+{
+	FSteamP2PConnectionInfo* ConnectionInfo = AcceptedConnections.Find(SessionId);
+	if (ConnectionInfo)
+	{
+        // Move active connections to the dead list so they can be removed (giving Steam a chance to flush connection)
+		DeadConnections.Add(SessionId, FSteamP2PConnectionInfo(ConnectionInfo->SteamNetworkingPtr, FPlatformTime::Seconds(), Channel));
+	}
+
+	UE_LOG_ONLINE(Log, TEXT("Removing P2P Session Id: %s, IdleTime: %0.3f"), *SessionId.ToDebugString(), ConnectionInfo ? (FPlatformTime::Seconds() - ConnectionInfo->LastReceivedTime) : 9999.f);
+	AcceptedConnections.Remove(SessionId);
+}
+
+/**
+ * Chance for the socket subsystem to get some time
+ *
+ * @param DeltaTime time since last tick
+ */
+bool FSocketSubsystemSteam::Tick(float DeltaTime)
+{	
+	double CurSeconds = FPlatformTime::Seconds();
+
+	// Debug connection state information
+	bool bDumpSessionInfo = false;
+	if ((CurSeconds - P2PDumpCounter) >= P2PDumpInterval)
+	{
+		P2PDumpCounter = CurSeconds;
+		bDumpSessionInfo = true;
+	}
+
+	TArray<FUniqueNetIdSteam> ExpiredSessions;
+	for (TMap<FUniqueNetIdSteam, FSteamP2PConnectionInfo>::TConstIterator It(AcceptedConnections); It; ++It)
+	{
+		const FUniqueNetIdSteam& SessionId = It.Key();
+		const FSteamP2PConnectionInfo& ConnectionInfo = It.Value();
+
+		bool bExpiredSession = true;
+		if (CurSeconds - ConnectionInfo.LastReceivedTime < P2PConnectionTimeout)
+		{
+			P2PSessionState_t SessionInfo;
+			if (ConnectionInfo.SteamNetworkingPtr != NULL && ConnectionInfo.SteamNetworkingPtr->GetP2PSessionState(SessionId, &SessionInfo))
+			{
+				bExpiredSession = false;
+
+				if (bDumpSessionInfo)
+				{
+					UE_LOG_ONLINE(Verbose, TEXT("Dumping Steam P2P socket details:"));
+					UE_LOG_ONLINE(Verbose, TEXT("- Id: %s, IdleTime: %0.3f"), *SessionId.ToDebugString(), (CurSeconds - ConnectionInfo.LastReceivedTime));
+
+					DumpSteamP2PSessionInfo(SessionInfo);
+				}
+			}
+			else
+			{
+				UE_LOG_ONLINE(Verbose, TEXT("Failed to get Steam P2P session state for Id: %s, IdleTime: %0.3f"), *SessionId.ToDebugString(), (CurSeconds - ConnectionInfo.LastReceivedTime));
+			}
+		}
+
+		if (bExpiredSession)
+		{
+			ExpiredSessions.Add(SessionId);
+		}
+	}
+
+	// Cleanup any closed/invalid sessions
+	for (int32 SessionIdx=0; SessionIdx < ExpiredSessions.Num(); SessionIdx++)
+	{
+		FUniqueNetIdSteam& SessionId = ExpiredSessions[SessionIdx];
+		P2PRemove(SessionId);
+	}
+
+	CleanupDeadConnections();
+
+	return true;
+}
+
+/**
+ * Iterate through the pending dead connections and permanently remove any that have been around
+ * long enough to flush their contents
+ */
+void FSocketSubsystemSteam::CleanupDeadConnections()
+{
+	double CurSeconds = FPlatformTime::Seconds();
+	TArray<FUniqueNetIdSteam> ExpiredSessions;
+	for (TMap<FUniqueNetIdSteam, FSteamP2PConnectionInfo>::TConstIterator It(DeadConnections); It; ++It)
+	{
+		const FUniqueNetIdSteam& SessionId = It.Key();
+		const FSteamP2PConnectionInfo& ConnectionInfo = It.Value();
+		if (CurSeconds - ConnectionInfo.LastReceivedTime >= 3.0f)
+		{
+			if (ConnectionInfo.Channel == -1)
+			{
+				ConnectionInfo.SteamNetworkingPtr->CloseP2PSessionWithUser(SessionId);
+			}
+			else
+			{
+				ConnectionInfo.SteamNetworkingPtr->CloseP2PChannelWithUser(SessionId, ConnectionInfo.Channel);
+			}
+
+			ExpiredSessions.Add(SessionId);
+		}
+	}
+
+	for (int32 SessionIdx=0; SessionIdx < ExpiredSessions.Num(); SessionIdx++)
+	{
+		FUniqueNetIdSteam& SessionId = ExpiredSessions[SessionIdx];
+		DeadConnections.Remove(SessionId);
+	}
+}
+
+/**
+ * Dumps the Steam P2P networking information for a given session id
+ *
+ * @param SessionInfo struct from Steam call to GetP2PSessionState
+ */
+void FSocketSubsystemSteam::DumpSteamP2PSessionInfo(P2PSessionState_t& SessionInfo)
+{
+	TSharedRef<FInternetAddr> IpAddr = ISocketSubsystem::Get(PLATFORM_SOCKETSUBSYSTEM)->CreateInternetAddr(SessionInfo.m_nRemoteIP, SessionInfo.m_nRemotePort);
+	UE_LOG_ONLINE(Verbose, TEXT("- Detailed P2P session info:"));
+	UE_LOG_ONLINE(Verbose, TEXT("-- IPAddress: %s"), *IpAddr->ToString(true));
+	UE_LOG_ONLINE(Verbose, TEXT("-- ConnectionActive: %i, Connecting: %i, SessionError: %i, UsingRelay: %i"),
+		SessionInfo.m_bConnectionActive, SessionInfo.m_bConnecting, SessionInfo.m_eP2PSessionError,
+		SessionInfo.m_bUsingRelay);
+	UE_LOG_ONLINE(Verbose, TEXT("-- QueuedBytes: %i, QueuedPackets: %i"), SessionInfo.m_nBytesQueuedForSend,
+		SessionInfo.m_nPacketsQueuedForSend);
+}
diff --git a/Engine/Plugins/Online/OnlineSubsystemSteam/Source/Private/SocketSubsystemSteam.h b/Engine/Plugins/Online/OnlineSubsystemSteam/Source/Private/SocketSubsystemSteam.h
new file mode 100644
--- /dev/null
+++ b/Engine/Plugins/Online/OnlineSubsystemSteam/Source/Private/SocketSubsystemSteam.h
@@ -0,0 +1,321 @@
+// Copyright 1998-2016 Epic Games, Inc. All Rights Reserved.
+
+#pragma once
+
+#include "SocketSubsystem.h"
+#include "IPAddress.h"
+#include "OnlineSubsystemSteamTypes.h"
+#include "OnlineSubsystemSteamPackage.h"
+
+/**
+ * Windows specific socket subsystem implementation
+ */
+class FSocketSubsystemSteam : public ISocketSubsystem, public FTickerObjectBase
+{
+protected:
+
+	/** Single instantiation of this subsystem */
+	static FSocketSubsystemSteam* SocketSingleton;
+
+	/** Tracks existing Steamworks sockets, for connection failure/timeout resolution */
+	TArray<class FSocketSteam*> SteamSockets;
+
+	/** Tracks existing Steamworks connections, for connection failure/timeout resolution */
+	TArray<struct FWeakObjectPtr> SteamConnections;
+
+    /** Keep track of Steam P2P connections */
+	struct FSteamP2PConnectionInfo
+	{
+        /** Steam networking interface responsible for this connection */
+		ISteamNetworking* SteamNetworkingPtr;
+		/** Channel to close */
+		int32 Channel;
+        /** Last time P2P session had activity (RecvFrom, etc) */
+		double LastReceivedTime;
+
+		FSteamP2PConnectionInfo() :
+			SteamNetworkingPtr(NULL),
+			Channel(-1),
+			LastReceivedTime(0.0)
+		{
+		}
+
+		FSteamP2PConnectionInfo(ISteamNetworking* InSteamNetworkingPtr, double InTime, int32 InChannel=-1) :
+			SteamNetworkingPtr(InSteamNetworkingPtr),
+			Channel(InChannel),
+			LastReceivedTime(InTime)
+		{
+		}
+	};
+
+    /** List of Steam P2P connections being tracked */
+	TMap<class FUniqueNetIdSteam, FSteamP2PConnectionInfo> AcceptedConnections;
+	/** List of Steam P2P connections to shutdown (dead connections remain around a few seconds longer to flush) */
+	TMap<class FUniqueNetIdSteam, FSteamP2PConnectionInfo> DeadConnections;
+
+	/**
+	 * Should Steam P2P sockets all fall back to Steam servers relay if a direct connection fails
+	 * read from [OnlineSubsystemSteam.bAllowP2PPacketRelay]
+	 */
+	bool bAllowP2PPacketRelay;
+	/** 
+     * Timeout period for any P2P session
+	 * read from [OnlineSubsystemSteam.P2PConnectionTimeout]
+     * (should be at least as long as NetDriver::ConnectionTimeout) 
+     */
+	float P2PConnectionTimeout;
+	/** Accumulated time before next dump of connection info */
+	double P2PDumpCounter;
+	/** Connection info output interval */
+	double P2PDumpInterval;
+
+	/**
+	 * Adds a steam socket for tracking
+	 *
+	 * @param InSocket	The socket to add for tracking
+	 */
+	void AddSocket(class FSocketSteam* InSocket)
+	{
+		SteamSockets.Add(InSocket);
+	}
+
+	/**
+	 * Removes a steam socket from tracking
+	 *
+	 * @param InSocket	The socket to remove from tracking
+	 */
+	void RemoveSocket(class FSocketSteam* InSocket)
+	{
+		SteamSockets.RemoveSingleSwap(InSocket);
+	}
+
+PACKAGE_SCOPE:
+
+	/** Last error set by the socket subsystem or one of its sockets */
+	int32 LastSocketError;
+
+	/** 
+	 * Singleton interface for this subsystem 
+	 * @return the only instance of this subsystem
+	 */
+	static FSocketSubsystemSteam* Create();
+
+	/**
+	 * Performs Windows specific socket clean up
+	 */
+	static void Destroy();
+
+	/**
+	 * Iterate through the pending dead connections and permanently remove any that have been around
+	 * long enough to flush their contents
+	 */
+	void CleanupDeadConnections();
+
+	/**
+	 * Associate the game server steam id with any sockets that were created prior to successful login
+	 *
+	 * @param GameServerId id assigned to this game server
+	 */
+	void FixupSockets(const FUniqueNetIdSteam& GameServerId);
+
+	/**
+	 * Adds a steam connection for tracking
+	 *
+	 * @param Connection The connection to add for tracking
+	 */
+	void RegisterConnection(class USteamNetConnection* Connection);
+
+	/**
+	 * Removes a steam connection from tracking
+	 *
+	 * @param Connection The connection to remove from tracking
+	 */
+	void UnregisterConnection(class USteamNetConnection* Connection);
+
+	/**
+	 * Notification from the Steam event layer that a remote connection has completely failed
+	 * 
+	 * @param RemoteId failure address
+	 */
+	void ConnectFailure(class FUniqueNetIdSteam& RemoteId);
+
+	/**
+	 * Potentially accept an incoming connection from a Steam P2P request
+	 * 
+	 * @param SteamNetworkingPtr the interface for access the P2P calls (Client/GameServer)
+	 * @param RemoteId the id of the incoming request
+	 * 
+	 * @return true if accepted, false otherwise
+	 */
+	bool AcceptP2PConnection(ISteamNetworking* SteamNetworkingPtr, FUniqueNetIdSteam& RemoteId);
+
+	/**
+	 * Add/update a Steam P2P connection as being recently accessed
+	 *
+	 * @param SteamNetworkingPtr proper networking interface that this session is communicating on
+	 * @param SessionId P2P session recently heard from
+     *
+     * @return true if the connection is active, false if this is in the dead connections list
+	 */
+	bool P2PTouch(ISteamNetworking* SteamNetworkingPtr, FUniqueNetIdSteam& SessionId);
+
+	/**
+	 * Remove a Steam P2P session from tracking and close the connection
+	 *
+	 * @param SessionId P2P session to remove
+	 * @param Channel channel to close, -1 to close all communication
+	 */
+	void P2PRemove(FUniqueNetIdSteam& SessionId, int32 Channel = -1);
+
+public:
+
+	FSocketSubsystemSteam() :
+	    bAllowP2PPacketRelay(false),
+		P2PConnectionTimeout(45.0f),
+		P2PDumpCounter(0.0),
+		P2PDumpInterval(3.0),
+		LastSocketError(0)
+	{
+	}
+
+	/**
+	 * Does Windows platform initialization of the sockets library
+	 *
+	 * @param Error a string that is filled with error information
+	 *
+	 * @return true if initialized ok, false otherwise
+	 */
+	virtual bool Init(FString& Error) override;
+
+	/**
+	 * Performs platform specific socket clean up
+	 */
+	virtual void Shutdown() override;
+
+	/**
+	 * Creates a socket
+	 *
+	 * @Param SocketType type of socket to create (DGram, Stream, etc)
+ 	 * @param SocketDescription debug description
+	 * @param bForceUDP overrides any platform specific protocol with UDP instead
+	 *
+	 * @return the new socket or NULL if failed
+	 */
+	virtual class FSocket* CreateSocket(const FName& SocketType, const FString& SocketDescription, bool bForceUDP = false) override;
+
+	/**
+	 * Cleans up a socket class
+	 *
+	 * @param Socket the socket object to destroy
+	 */
+	virtual void DestroySocket(class FSocket* Socket) override;
+
+	/**
+	 * Does a DNS look up of a host name
+	 *
+	 * @param HostName the name of the host to look up
+	 * @param OutAddr the address to copy the IP address to
+	 */
+	virtual ESocketErrors GetHostByName(const ANSICHAR* HostName, FInternetAddr& OutAddr) override;
+
+	/**
+	 * Some platforms require chat data (voice, text, etc.) to be placed into
+	 * packets in a special way. This function tells the net connection
+	 * whether this is required for this platform
+	 */
+	virtual bool RequiresChatDataBeSeparate() override
+	{
+		return false;
+	}
+
+	/**
+	 * Some platforms require packets be encrypted. This function tells the
+	 * net connection whether this is required for this platform
+	 */
+	virtual bool RequiresEncryptedPackets() override
+	{
+		return false;
+	}
+
+	/**
+	 * Determines the name of the local machine
+	 *
+	 * @param HostName the string that receives the data
+	 *
+	 * @return true if successful, false otherwise
+	 */
+	virtual bool GetHostName(FString& HostName) override;
+
+	/**
+	 *	Create a proper FInternetAddr representation
+	 * @param Address host address
+	 * @param Port host port
+	 */
+	virtual TSharedRef<FInternetAddr> CreateInternetAddr(uint32 Address=0, uint32 Port=0) override;
+
+	/**
+	 * @return Whether the machine has a properly configured network device or not
+	 */
+	virtual bool HasNetworkDevice() override;
+
+	/**
+	 *	Get the name of the socket subsystem
+	 * @return a string naming this subsystem
+	 */
+	virtual const TCHAR* GetSocketAPIName() const override;
+
+	/**
+	 * Returns the last error that has happened
+	 */
+	virtual ESocketErrors GetLastErrorCode() override;
+
+	/**
+	 * Translates the platform error code to a ESocketErrors enum
+	 */
+	virtual ESocketErrors TranslateErrorCode(int32 Code) override;
+
+	/**
+	 * Gets the list of addresses associated with the adapters on the local computer.
+	 *
+	 * @param OutAdresses - Will hold the address list.
+	 *
+	 * @return true on success, false otherwise.
+	 */
+	virtual bool GetLocalAdapterAddresses( TArray<TSharedPtr<FInternetAddr> >& OutAdresses ) override
+	{
+		bool bCanBindAll;
+
+		OutAdresses.Add(GetLocalHostAddr(*GLog, bCanBindAll));
+
+		return true;
+	}
+
+	/**
+	 *	Get local IP to bind to
+	 */
+	virtual TSharedRef<FInternetAddr> GetLocalBindAddr(FOutputDevice& Out) override;
+
+	/**
+	 * Dumps the Steam P2P networking information for a given session state
+	 *
+	 * @param SessionInfo struct from Steam call to GetP2PSessionState
+	 */
+	void DumpSteamP2PSessionInfo(P2PSessionState_t& SessionInfo);
+
+	/**
+	 * Chance for the socket subsystem to get some time
+	 *
+	 * @param DeltaTime time since last tick
+	 */
+	virtual bool Tick(float DeltaTime) override;
+};
+
+/**
+ * Create the socket subsystem for the given platform service
+ */
+FName CreateSteamSocketSubsystem();
+
+/**
+ * Tear down the socket subsystem for the given platform service
+ */
+void DestroySteamSocketSubsystem();
\ No newline at end of file
diff --git a/Engine/Plugins/Online/OnlineSubsystemSteam/Source/Private/SteamNetConnection.cpp b/Engine/Plugins/Online/OnlineSubsystemSteam/Source/Private/SteamNetConnection.cpp
new file mode 100644
--- /dev/null
+++ b/Engine/Plugins/Online/OnlineSubsystemSteam/Source/Private/SteamNetConnection.cpp
@@ -0,0 +1,89 @@
+// Copyright 1998-2016 Epic Games, Inc. All Rights Reserved.
+
+#include "OnlineSubsystemSteamPrivatePCH.h"
+#include "Classes/SteamNetConnection.h"
+#include "IPAddressSteam.h"
+#include "SocketSubsystemSteam.h"
+#include "SocketsSteam.h"
+#include "SteamNetDriver.h"
+
+USteamNetConnection::USteamNetConnection(const FObjectInitializer& ObjectInitializer) :
+	Super(ObjectInitializer),
+	bIsPassthrough(false)
+{
+}
+
+void USteamNetConnection::InitLocalConnection(UNetDriver* InDriver, class FSocket* InSocket, const FURL& InURL, EConnectionState InState, int32 InMaxPacket, int32 InPacketOverhead)
+{
+	bIsPassthrough = InURL.Host.StartsWith(STEAM_URL_PREFIX) ? false : true;
+	
+	Super::InitLocalConnection(InDriver, InSocket, InURL, InState, InMaxPacket, InPacketOverhead);
+	if (!bIsPassthrough && RemoteAddr.IsValid())
+	{
+		FSocketSubsystemSteam* SocketSubsystem = (FSocketSubsystemSteam*)ISocketSubsystem::Get(STEAM_SUBSYSTEM);
+		if (SocketSubsystem)
+		{
+			SocketSubsystem->RegisterConnection(this);
+		}
+	}
+}
+
+void USteamNetConnection::InitRemoteConnection(UNetDriver* InDriver, class FSocket* InSocket, const FURL& InURL, const class FInternetAddr& InRemoteAddr, EConnectionState InState, int32 InMaxPacket, int32 InPacketOverhead)
+{
+	bIsPassthrough = ((USteamNetDriver*)InDriver)->bIsPassthrough;
+
+	Super::InitRemoteConnection(InDriver, InSocket, InURL, InRemoteAddr, InState, InMaxPacket, InPacketOverhead);
+	if (!bIsPassthrough && RemoteAddr.IsValid())
+	{
+		FSocketSubsystemSteam* SocketSubsystem = (FSocketSubsystemSteam*)ISocketSubsystem::Get(STEAM_SUBSYSTEM);
+		if (SocketSubsystem)
+		{
+			SocketSubsystem->RegisterConnection(this);
+		}
+	}
+}
+
+void USteamNetConnection::CleanUp()
+{
+	Super::CleanUp();
+
+	if (!bIsPassthrough)
+	{
+		FSocketSubsystemSteam* SocketSubsystem = (FSocketSubsystemSteam*)ISocketSubsystem::Get(STEAM_SUBSYSTEM);
+		if (SocketSubsystem)
+		{
+			// Unregister the connection AFTER the parent class has had a chance to close/flush connections
+			SocketSubsystem->UnregisterConnection(this);
+		}
+	}
+}
+
+void USteamNetConnection::DumpSteamConnection()
+{
+	if (!bIsPassthrough && RemoteAddr.IsValid())
+	{
+		UE_LOG_ONLINE(Verbose, TEXT("%s: Dumping Steam P2P connection details:"), *GetName());
+		FInternetAddrSteam& SteamAddr = (FInternetAddrSteam&)(*RemoteAddr); 
+		FUniqueNetIdSteam& SessionId = SteamAddr.SteamId;
+		UE_LOG_ONLINE(Verbose, TEXT("- Id: %s, IdleTime: %0.3f"), *SessionId.ToDebugString(), (Driver->Time - LastReceiveTime));
+	
+		if (Socket)
+		{
+			FSocketSteam* SteamSocket = (FSocketSteam*)Socket;
+
+			FSocketSubsystemSteam* SteamSockets = (FSocketSubsystemSteam*)ISocketSubsystem::Get(STEAM_SUBSYSTEM);
+			if (SteamSockets)
+			{
+				P2PSessionState_t SessionInfo;
+				if (SteamSocket->SteamNetworkingPtr != NULL && SteamSocket->SteamNetworkingPtr->GetP2PSessionState(SessionId, &SessionInfo))
+				{
+					SteamSockets->DumpSteamP2PSessionInfo(SessionInfo);
+				}
+				else
+				{
+					UE_LOG_ONLINE(Verbose, TEXT("Failed to get Steam P2P session state for Id: %s"), *SessionId.ToDebugString());
+				}
+			}
+		}
+	}
+}
\ No newline at end of file
diff --git a/Engine/Plugins/Online/OnlineSubsystemSteam/Source/Private/SteamNetDriver.cpp b/Engine/Plugins/Online/OnlineSubsystemSteam/Source/Private/SteamNetDriver.cpp
new file mode 100644
--- /dev/null
+++ b/Engine/Plugins/Online/OnlineSubsystemSteam/Source/Private/SteamNetDriver.cpp
@@ -0,0 +1,173 @@
+// Copyright 1998-2016 Epic Games, Inc. All Rights Reserved.
+
+#include "OnlineSubsystemSteamPrivatePCH.h"
+#include "SteamNetDriver.h"
+#include "IPAddressSteam.h"
+#include "OnlineSubsystemSteam.h"
+#include "SocketSubsystemSteam.h"
+#include "SocketsSteam.h"
+#include "SteamNetDriver.h"
+#include "SteamNetConnection.h"
+
+USteamNetDriver::USteamNetDriver(const FObjectInitializer& ObjectInitializer) :
+	Super(ObjectInitializer),
+	bIsPassthrough(false)
+{
+}
+
+void USteamNetDriver::PostInitProperties()
+{
+	Super::PostInitProperties();
+	ConnectionDumpInterval = 10.0f;
+	ConnectionDumpCounter = 0.0f;
+}
+
+bool USteamNetDriver::IsAvailable() const
+{
+	// Net driver won't work if the online and socket subsystems don't exist
+	IOnlineSubsystem* SteamSubsystem = IOnlineSubsystem::Get(STEAM_SUBSYSTEM);
+	if (SteamSubsystem)
+	{
+		ISocketSubsystem* SteamSockets = ISocketSubsystem::Get(STEAM_SUBSYSTEM);
+		if (SteamSockets)
+		{
+			return true;
+		}
+	}
+
+	return false;
+}
+
+ISocketSubsystem* USteamNetDriver::GetSocketSubsystem()
+{
+	return ISocketSubsystem::Get(bIsPassthrough ? PLATFORM_SOCKETSUBSYSTEM : STEAM_SUBSYSTEM);
+}
+
+bool USteamNetDriver::InitBase(bool bInitAsClient, FNetworkNotify* InNotify, const FURL& URL, bool bReuseAddressAndPort, FString& Error)
+{
+	if (bIsPassthrough)
+	{
+		return UIpNetDriver::InitBase(bInitAsClient, InNotify, URL, bReuseAddressAndPort, Error);
+	}
+
+	// Skip UIpNetDriver implementation
+	if (!UNetDriver::InitBase(bInitAsClient, InNotify, URL, bReuseAddressAndPort, Error))
+	{
+		return false;
+	}
+
+	ISocketSubsystem* SocketSubsystem = GetSocketSubsystem();
+	if (SocketSubsystem == NULL)
+	{
+		UE_LOG(LogNet, Warning, TEXT("Unable to find socket subsystem"));
+		Error = TEXT("Unable to find socket subsystem");
+		return false;
+	}
+
+	if(Socket == NULL)
+	{
+		Socket = 0;
+		Error = FString::Printf( TEXT("SteamSockets: socket failed (%i)"), (int32)SocketSubsystem->GetLastErrorCode() );
+		return false;
+	}
+
+	// Bind socket to our port.
+	LocalAddr = SocketSubsystem->GetLocalBindAddr(*GLog);
+
+	// Set the Steam channel (port) to communicate on (both Client/Server communicate on same "channel")
+	LocalAddr->SetPort(URL.Port);
+
+	int32 AttemptPort = LocalAddr->GetPort();
+	int32 BoundPort = SocketSubsystem->BindNextPort(Socket, *LocalAddr, MaxPortCountToTry + 1, 1);
+	UE_LOG(LogNet, Display, TEXT("%s bound to port %d"), *GetName(), BoundPort);
+	// Success.
+	return true;
+}
+
+bool USteamNetDriver::InitConnect(FNetworkNotify* InNotify, const FURL& ConnectURL, FString& Error)
+{
+	ISocketSubsystem* SteamSockets = ISocketSubsystem::Get(STEAM_SUBSYSTEM);
+	if (SteamSockets)
+	{
+		// If we are opening a Steam URL, create a Steam client socket
+		if (ConnectURL.Host.StartsWith(STEAM_URL_PREFIX))
+		{
+			Socket = SteamSockets->CreateSocket(FName(TEXT("SteamClientSocket")), TEXT("Unreal client (Steam)"));
+		}
+		else
+		{
+			bIsPassthrough = true;
+		}
+	}
+
+	return Super::InitConnect(InNotify, ConnectURL, Error);
+}
+
+bool USteamNetDriver::InitListen(FNetworkNotify* InNotify, FURL& ListenURL, bool bReuseAddressAndPort, FString& Error)
+{
+	ISocketSubsystem* SteamSockets = ISocketSubsystem::Get(STEAM_SUBSYSTEM);
+	if (SteamSockets && !ListenURL.HasOption(TEXT("bIsLanMatch")))
+	{
+		FName SocketTypeName = IsRunningDedicatedServer() ? FName(TEXT("SteamServerSocket")) : FName(TEXT("SteamClientSocket"));
+		Socket = SteamSockets->CreateSocket(SocketTypeName, TEXT("Unreal server (Steam)"));
+	}
+	else
+	{
+        // Socket will be created in base class
+		bIsPassthrough = true;
+	}
+
+	return Super::InitListen(InNotify, ListenURL, bReuseAddressAndPort, Error);
+}
+
+void USteamNetDriver::Shutdown()
+{
+	if (!bIsPassthrough)
+	{
+		FSocketSteam* SteamSocket = (FSocketSteam*)Socket;
+		if (SteamSocket)
+		{
+			SteamSocket->SetSteamSendMode(k_EP2PSendUnreliableNoDelay);
+		}
+	}
+
+	Super::Shutdown();
+}
+
+void USteamNetDriver::TickFlush(float DeltaSeconds)
+{
+	Super::TickFlush(DeltaSeconds);
+
+	if (!bIsPassthrough)
+	{
+		// Debug connection state information
+		double CurSeconds = FPlatformTime::Seconds();
+		if ((CurSeconds - ConnectionDumpCounter) >= ConnectionDumpInterval)
+		{
+			ConnectionDumpCounter = CurSeconds;
+
+			USteamNetConnection* ServerConn = Cast<USteamNetConnection>(ServerConnection);
+			if (ServerConn != NULL)
+			{
+				ServerConn->DumpSteamConnection();	
+			}
+
+			USteamNetConnection* ClientConn = NULL;
+			for (int32 ConnIdx=0; ConnIdx < ClientConnections.Num(); ConnIdx++)
+			{
+				ClientConn = Cast<USteamNetConnection>(ClientConnections[ConnIdx]);
+				if (ClientConn != NULL)
+				{
+					ClientConn->DumpSteamConnection();
+				}
+			}
+		}
+	}
+}
+
+bool USteamNetDriver::IsNetResourceValid()
+{
+	bool bIsValidSteamSocket = !bIsPassthrough && (Socket != NULL) && ((FSocketSteam*)Socket)->LocalSteamId.IsValid();
+	bool bIsValidPassthroughSocket = bIsPassthrough && UIpNetDriver::IsNetResourceValid();
+	return bIsValidSteamSocket || bIsValidPassthroughSocket;
+}
diff --git a/Engine/Plugins/Online/OnlineSubsystemSteam/Source/Private/VoiceEngineSteam.cpp b/Engine/Plugins/Online/OnlineSubsystemSteam/Source/Private/VoiceEngineSteam.cpp
new file mode 100644
--- /dev/null
+++ b/Engine/Plugins/Online/OnlineSubsystemSteam/Source/Private/VoiceEngineSteam.cpp
@@ -0,0 +1,395 @@
+// Copyright 1998-2016 Epic Games, Inc. All Rights Reserved.
+
+#include "OnlineSubsystemSteamPrivatePCH.h"
+#include "OnlineSubsystemSteamTypes.h"
+#include "VoiceEngineSteam.h"
+#include "OnlineSubsystemSteam.h"
+#include "Voice.h"
+#include "SteamUtilities.h"
+
+#include "Engine.h"
+#include "SoundDefinitions.h"
+#include "Runtime/Engine/Classes/Sound/SoundWaveProcedural.h"
+#include "OnlineSubsystemUtils.h"
+
+/** Largest size Steam says it will need to compress data */
+#define MAX_COMPRESSED_VOICE_BUFFER_SIZE 8 * 1024
+/** Largest size Steam says it will uncompress to */
+#define MAX_UNCOMPRESSED_VOICE_BUFFER_SIZE 22 * 1024
+
+FVoiceEngineSteam::FVoiceEngineSteam(FOnlineSubsystemSteam* InSteamSubsystem) :
+	SteamSubsystem(InSteamSubsystem),
+	OwningUserIndex(INVALID_INDEX),
+	bPendingFinalCapture(false),
+	bIsCapturing(false),
+	SerializeHelper(NULL)
+{
+	SteamUserPtr = SteamUser();
+	SteamFriendsPtr = SteamFriends();
+}
+
+FVoiceEngineSteam::~FVoiceEngineSteam()
+{
+	if (bIsCapturing)
+	{
+		SteamUserPtr->StopVoiceRecording();
+		SteamFriendsPtr->SetInGameVoiceSpeaking(SteamUserPtr->GetSteamID(), false);
+	}
+	delete SerializeHelper;
+}
+
+void FVoiceEngineSteam::VoiceCaptureUpdate() const
+{
+	if (bPendingFinalCapture)
+	{
+		uint32 CompressedSize;
+		const EVoiceResult RecordingState = SteamUserPtr->GetAvailableVoice(&CompressedSize, NULL, 0);
+
+		// If no data is available, we have finished capture the last (post-StopRecording) half-second of voice data
+		if (RecordingState == k_EVoiceResultNotRecording)
+		{
+			UE_LOG(LogVoiceEncode, Log, TEXT("Internal voice capture complete."));
+
+			bPendingFinalCapture = false;
+
+			// If a new recording session has begun since the call to 'StopRecording', kick that off
+			if (bIsCapturing)
+			{
+				StartRecording();
+			}
+			else
+			{
+				// Marks that recording has successfully stopped
+				StoppedRecording();
+			}
+		}
+	}
+}
+
+void FVoiceEngineSteam::StartRecording() const
+{
+	UE_LOG(LogVoiceEncode, VeryVerbose, TEXT("VOIP StartRecording"));
+	if (SteamUserPtr)
+	{
+		SteamUserPtr->StartVoiceRecording();
+		SteamFriendsPtr->SetInGameVoiceSpeaking(SteamUserPtr->GetSteamID(), true);
+	}
+}
+
+void FVoiceEngineSteam::StopRecording() const
+{
+	UE_LOG(LogVoiceEncode, VeryVerbose, TEXT("VOIP StopRecording"));
+	if (SteamUserPtr)
+	{
+		SteamUserPtr->StopVoiceRecording();
+	}
+}
+
+void FVoiceEngineSteam::StoppedRecording() const
+{
+	UE_LOG(LogVoiceEncode, VeryVerbose, TEXT("VOIP StoppedRecording"));
+	if (SteamFriendsPtr)
+	{
+		SteamFriendsPtr->SetInGameVoiceSpeaking(SteamUserPtr->GetSteamID(), false);
+	}
+}
+
+bool FVoiceEngineSteam::Init(int32 MaxLocalTalkers, int32 MaxRemoteTalkers) 
+{
+	bool bSuccess = false;
+
+	if (SteamSubsystem->IsSteamClientAvailable())
+	{
+		bool bHasVoiceEnabled = false;
+		if (GConfig->GetBool(TEXT("OnlineSubsystem"),TEXT("bHasVoiceEnabled"), bHasVoiceEnabled, GEngineIni) && bHasVoiceEnabled)
+		{
+			if (SteamUserPtr && SteamFriendsPtr)
+			{
+				// Just verify voice capture is available
+				uint32 CompressedBytes;
+				SteamUserPtr->StartVoiceRecording();
+				EVoiceResult VoiceResult = SteamUserPtr->GetAvailableVoice(&CompressedBytes, NULL, 0);
+				SteamUserPtr->StopVoiceRecording();
+
+				bSuccess = (VoiceResult != k_EVoiceResultNotInitialized);
+				if (bSuccess)
+				{
+					CompressedVoiceBuffer.Empty(MAX_COMPRESSED_VOICE_BUFFER_SIZE);
+					DecompressedVoiceBuffer.Empty(MAX_UNCOMPRESSED_VOICE_BUFFER_SIZE);
+				}
+				else
+				{
+					UE_LOG(LogVoice, Warning, TEXT("Steamworks voice initialization failed!"));
+				}
+			}
+		}
+		else
+		{
+			UE_LOG(LogVoice, Log, TEXT("Voice module disabled by config [OnlineSubsystem].bHasVoiceEnabled"));
+		}
+	}
+
+	return bSuccess;
+}
+
+uint32 FVoiceEngineSteam::StartLocalVoiceProcessing(uint32 LocalUserNum) 
+{
+	uint32 Return = E_FAIL;
+	if (IsOwningUser(LocalUserNum))
+	{
+		if (!bIsCapturing)
+		{
+			// Update the current recording state, if VOIP data was still being read
+			VoiceCaptureUpdate();
+
+			if (!IsRecording())
+			{
+				StartRecording();
+			}
+
+			bIsCapturing = true;
+		}
+
+		Return = S_OK;
+	}
+	else
+	{
+		UE_LOG(LogVoiceEncode, Error, TEXT("StartLocalVoiceProcessing(): Device is currently owned by another user"));
+	}
+
+	return Return;
+}
+
+uint32 FVoiceEngineSteam::StopLocalVoiceProcessing(uint32 LocalUserNum) 
+{
+	uint32 Return = E_FAIL;
+	if (IsOwningUser(LocalUserNum))
+	{
+		if (bIsCapturing)
+		{
+			bIsCapturing = false;
+			bPendingFinalCapture = true;
+
+			// Make a call to begin stopping the current VOIP recording session
+			StopRecording();
+
+			// Now check/update the status of the recording session
+			VoiceCaptureUpdate();
+		}
+
+		Return = S_OK;
+	}
+	else
+	{
+		UE_LOG(LogVoiceEncode, Error, TEXT("StopLocalVoiceProcessing: Ignoring stop request for non-owning user"));
+	}
+
+	return Return;
+}
+
+uint32 FVoiceEngineSteam::GetVoiceDataReadyFlags() const 
+{
+	// First check and update the internal state of VOIP recording
+	VoiceCaptureUpdate();
+	if (OwningUserIndex != INVALID_INDEX && IsRecording())
+	{
+		// Check if there is new data available via the Steam API
+		if (AvailableVoiceResult == k_EVoiceResultOK && CompressedBytesAvailable > 0)
+		{
+			return 1 << OwningUserIndex;
+		}
+	}
+
+	return 0;
+}
+
+uint32 FVoiceEngineSteam::ReadLocalVoiceData(uint32 LocalUserNum, uint8* Data, uint32* Size) 
+{
+	check(*Size > 0);
+
+	// Before doing anything, check/update the current recording state
+	VoiceCaptureUpdate();
+
+	// Return data even if not capturing, if the final half-second of data from Steam is still pending
+	if (IsOwningUser(LocalUserNum) && IsRecording())
+	{
+		CompressedVoiceBuffer.Empty(MAX_COMPRESSED_VOICE_BUFFER_SIZE);
+
+		uint32 CompressedBytes = 0;
+		EVoiceResult VoiceResult = SteamUserPtr->GetAvailableVoice(&CompressedBytes, NULL, 0);
+		if (VoiceResult != k_EVoiceResultOK && VoiceResult != k_EVoiceResultNoData)
+		{
+			UE_LOG(LogVoiceEncode, Warning, TEXT("ReadLocalVoiceData: GetAvailableVoice failure: VoiceResult: %s"), *SteamVoiceResult(VoiceResult));
+			return E_FAIL;
+		}
+
+		// This shouldn't happen, but just in case.
+		if (CompressedBytes == 0)
+		{
+			UE_LOG(LogVoiceEncode, VeryVerbose, TEXT("ReadLocalVoiceData: No Data: VoiceResult: %s"), *SteamVoiceResult(VoiceResult));
+			*Size = 0;
+			return S_OK;
+		}
+
+		// Update the amount of data available for consumption (CompressedBytes can be more than AvailableWritten below)
+		CompressedVoiceBuffer.AddUninitialized(CompressedBytes);
+
+		uint32 AvailableWritten = 0;
+		VoiceResult = SteamUserPtr->GetVoice(true, CompressedVoiceBuffer.GetData(), CompressedVoiceBuffer.Num(), &AvailableWritten, false, NULL, 0, NULL, 0);
+
+		static double LastGetVoiceCallTime = 0.0;
+		double CurTime = FPlatformTime::Seconds();
+		double TimeSinceLastCall = LastGetVoiceCallTime > 0 ? (CurTime - LastGetVoiceCallTime) : 0.0;
+		LastGetVoiceCallTime = CurTime;
+
+		UE_LOG(LogVoiceEncode, Log, TEXT("ReadLocalVoiceData: GetVoice: Result: %s, Available: %i, LastCall: %0.3f"), *SteamVoiceResult(VoiceResult), AvailableWritten, TimeSinceLastCall);
+		if (VoiceResult != k_EVoiceResultOK)
+		{
+			UE_LOG(LogVoiceEncode, Warning, TEXT("ReadLocalVoiceData: GetVoice failure: VoiceResult: %s"), *SteamVoiceResult(VoiceResult));
+			*Size = 0;
+			CompressedVoiceBuffer.Empty(MAX_COMPRESSED_VOICE_BUFFER_SIZE);
+			return E_FAIL;
+		}
+
+		if (AvailableWritten > 0)
+		{
+			*Size = FMath::Min<int32>(*Size, AvailableWritten);
+			FMemory::Memcpy(Data, CompressedVoiceBuffer.GetData(), *Size);
+
+			UE_LOG(LogVoiceEncode, VeryVerbose, TEXT("ReadLocalVoiceData: Size: %d"), *Size);
+			return S_OK;
+		}
+		else
+		{
+			*Size = 0;
+		}
+	}
+
+	return E_FAIL;
+}
+
+uint32 FVoiceEngineSteam::SubmitRemoteVoiceData(const FUniqueNetId& RemoteTalkerId, uint8* Data, uint32* Size) 
+{
+	UE_LOG(LogVoiceDecode, VeryVerbose, TEXT("SubmitRemoteVoiceData(%s) Size: %d received!"), *RemoteTalkerId.ToDebugString(), *Size);
+
+	const FUniqueNetIdSteam& SteamId = (const FUniqueNetIdSteam&)RemoteTalkerId;
+	FRemoteTalkerDataSteam* QueuedData = RemoteTalkerBuffers.Find(SteamId);
+
+	if (QueuedData == NULL)
+	{
+		RemoteTalkerBuffers.Add(SteamId, FRemoteTalkerDataSteam());
+		QueuedData = RemoteTalkerBuffers.Find(SteamId);
+	}
+
+	check(QueuedData);
+
+	// new voice packet.
+	QueuedData->LastSeen = FPlatformTime::Seconds();
+
+	uint32 BytesWritten = 0;
+
+	DecompressedVoiceBuffer.Empty(MAX_UNCOMPRESSED_VOICE_BUFFER_SIZE);
+	DecompressedVoiceBuffer.AddUninitialized(MAX_UNCOMPRESSED_VOICE_BUFFER_SIZE);
+	const EVoiceResult VoiceResult = SteamUserPtr->DecompressVoice(Data, *Size, DecompressedVoiceBuffer.GetData(),
+		DecompressedVoiceBuffer.Num(), &BytesWritten, SteamUserPtr->GetVoiceOptimalSampleRate());
+
+	if (VoiceResult != k_EVoiceResultOK)
+	{
+		UE_LOG(LogVoiceDecode, Warning, TEXT("SubmitRemoteVoiceData: DecompressVoice failure: VoiceResult: %s"), *SteamVoiceResult(VoiceResult));
+		*Size = 0;
+		return E_FAIL;
+	}
+
+	// If there is no data, return
+	if (BytesWritten <= 0)
+	{
+		*Size = 0;
+		return S_OK;
+	}
+
+	// Generate a streaming wave audio component for voice playback
+	if (QueuedData->AudioComponent == NULL || QueuedData->AudioComponent->IsPendingKill())
+	{
+		if (SerializeHelper == NULL)
+		{
+			SerializeHelper = new FVoiceSerializeHelper(this);
+		}
+
+		QueuedData->AudioComponent = CreateVoiceAudioComponent(SteamUserPtr->GetVoiceOptimalSampleRate());
+		if (QueuedData->AudioComponent)
+		{
+			QueuedData->AudioComponent->OnAudioFinishedNative.AddRaw(this, &FVoiceEngineSteam::OnAudioFinished);
+			QueuedData->AudioComponent->Play();
+		}
+	}
+
+	if (QueuedData->AudioComponent != NULL)
+	{
+		USoundWaveProcedural* SoundStreaming = CastChecked<USoundWaveProcedural>(QueuedData->AudioComponent->Sound);
+		if (SoundStreaming->GetAvailableAudioByteCount() == 0)
+		{
+			UE_LOG(LogVoiceDecode, Log, TEXT("VOIP audio component was starved!"));
+		}
+		SoundStreaming->QueueAudio(DecompressedVoiceBuffer.GetData(), BytesWritten);
+	}
+
+	return S_OK;
+}
+
+void FVoiceEngineSteam::TickTalkers(float DeltaTime)
+{
+	// Remove users that are done talking.
+	const double CurTime = FPlatformTime::Seconds();
+	for (FRemoteTalkerData::TIterator It(RemoteTalkerBuffers); It; ++It)
+	{
+		FRemoteTalkerDataSteam& RemoteData = It.Value();
+		double TimeSince = CurTime - RemoteData.LastSeen;
+		if (TimeSince >= 5.0)
+		{
+			// Dump the whole talker
+			if (RemoteData.AudioComponent)
+			{
+				RemoteData.AudioComponent->Stop();
+				RemoteData.AudioComponent = NULL;
+			}
+
+			It.RemoveCurrent();
+		}
+	}
+}
+
+void FVoiceEngineSteam::Tick(float DeltaTime) 
+{
+	// Check available voice one a frame, this value changes after calling GetVoice()
+	AvailableVoiceResult = SteamUserPtr->GetAvailableVoice(&CompressedBytesAvailable, NULL, 0);
+	//UE_LOG(LogVoiceEncode, Log, TEXT("TickResult: %s"), *SteamVoiceResult(AvailableVoiceResult));
+	TickTalkers(DeltaTime);
+}
+
+void FVoiceEngineSteam::OnAudioFinished(UAudioComponent* AC)
+{
+	for (FRemoteTalkerData::TIterator It(RemoteTalkerBuffers); It; ++It)
+	{
+		FRemoteTalkerDataSteam& RemoteData = It.Value();
+		if (RemoteData.AudioComponent->IsPendingKill() || AC == RemoteData.AudioComponent)
+		{
+			UE_LOG(LogVoiceDecode, Log, TEXT("Removing VOIP AudioComponent for SteamId: %s"), *It.Key().ToDebugString());
+			RemoteData.AudioComponent = NULL;
+			break;
+		}
+	}
+	UE_LOG(LogVoiceDecode, Verbose, TEXT("Audio Finished"));
+}
+
+FString FVoiceEngineSteam::GetVoiceDebugState() const 
+{
+	FString Output;
+	Output = FString::Printf(TEXT("IsRecording: %d\n DataReady: 0x%08x State:%s\n BufferRemaining: %d\n"),
+		IsRecording(), 
+		GetVoiceDataReadyFlags(),
+		*SteamVoiceResult(AvailableVoiceResult),
+		CompressedVoiceBuffer.Num()
+		);
+
+	return Output;
+}
\ No newline at end of file
diff --git a/Engine/Plugins/Online/OnlineSubsystemSteam/Source/Private/VoiceEngineSteam.h b/Engine/Plugins/Online/OnlineSubsystemSteam/Source/Private/VoiceEngineSteam.h
new file mode 100644
--- /dev/null
+++ b/Engine/Plugins/Online/OnlineSubsystemSteam/Source/Private/VoiceEngineSteam.h
@@ -0,0 +1,234 @@
+// Copyright 1998-2016 Epic Games, Inc. All Rights Reserved.
+
+#pragma once
+
+#include "VoiceInterface.h"
+#include "Net/VoiceDataCommon.h"
+#include "OnlineSubsystemTypes.h"
+#include "OnlineSubsystemSteamPackage.h"
+#include "OnlineSubsystemSteamTypes.h"
+
+/** 
+ * Remote voice data playing on a single client
+ */
+class FRemoteTalkerDataSteam
+{
+public:
+	FRemoteTalkerDataSteam() :
+		LastSeen(0.0),
+		AudioComponent(NULL)
+	{
+	}
+
+	// Receive side timestamp since last voice packet fragment
+	double LastSeen;
+	// Audio component playing this buffer (only valid on remote instances)
+	class UAudioComponent* AudioComponent;
+};
+
+/**
+ * The Steam implementation of the voice engine 
+ */
+class FVoiceEngineSteam : public IVoiceEngine
+{
+ 	class FVoiceSerializeHelper : public FGCObject
+ 	{
+ 		/** Reference to audio components */
+ 		FVoiceEngineSteam* VoiceEngine;
+ 		FVoiceSerializeHelper() :
+ 			VoiceEngine(NULL)
+ 		{}
+ 
+ 	public:
+ 
+ 		FVoiceSerializeHelper(FVoiceEngineSteam* InVoiceEngine) :
+ 			VoiceEngine(InVoiceEngine)
+ 		{}
+ 		~FVoiceSerializeHelper() {}
+ 		
+ 		/** FGCObject interface */
+ 		virtual void AddReferencedObjects(FReferenceCollector& Collector) override
+ 		{
+ 			// Prevent garbage collection of audio components
+ 			for (FRemoteTalkerData::TIterator It(VoiceEngine->RemoteTalkerBuffers); It; ++It)
+ 			{
+ 				FRemoteTalkerDataSteam& RemoteData = It.Value();
+ 				if (RemoteData.AudioComponent)
+ 				{
+ 					Collector.AddReferencedObject(RemoteData.AudioComponent);
+ 				}
+ 			}
+ 		}
+ 	};
+
+	friend class FVoiceSerializeHelper;
+
+	/** Mapping of UniqueIds to the incoming voice data and their audio component */
+	typedef TMap<FUniqueNetIdSteam, FRemoteTalkerDataSteam> FRemoteTalkerData;
+
+	/** Reference to the main Steam subsystem */
+	class FOnlineSubsystemSteam* SteamSubsystem;
+    /** Steam User interface */
+	class ISteamUser* SteamUserPtr;
+    /** Steam Friends interface */
+	class ISteamFriends* SteamFriendsPtr;
+    /** User index currently holding onto the voice interface */
+	int32 OwningUserIndex;
+	/** Amount of compressed data available this frame */
+	uint32 CompressedBytesAvailable;
+	/** Result of call to GetAvailableVoice() this frame*/
+	EVoiceResult AvailableVoiceResult;
+    /** Have we stopped Steam voice but are waiting for its completion */
+	mutable bool bPendingFinalCapture;
+    /** State of voice recording */
+	bool bIsCapturing;
+
+	/** Data from Steamworks, waiting to send to network. */
+	TArray<uint8> CompressedVoiceBuffer;
+	/** Data from network playing on an audio component. */
+	FRemoteTalkerData RemoteTalkerBuffers;
+	/** Voice decompression buffer, shared by all talkers */
+	TArray<uint8> DecompressedVoiceBuffer;
+	/** Serialization helper */
+	class FVoiceSerializeHelper* SerializeHelper;
+
+	/**
+	 * Determines if the specified index is the owner or not
+	 *
+	 * @param InIndex the index being tested
+	 *
+	 * @return true if this is the owner, false otherwise
+	 */
+	FORCEINLINE bool IsOwningUser(uint32 UserIndex)
+	{
+		return UserIndex >= 0 && UserIndex < MAX_SPLITSCREEN_TALKERS && OwningUserIndex == UserIndex;
+	}
+
+	/** 
+	 * Update the internal state of the voice capturing state 
+	 * Handles Steam's continual recording for "last half second" after requested stop
+	 */
+	void VoiceCaptureUpdate() const;
+
+	/** Tell Steam to start capturing voice data */
+	void StartRecording() const;
+
+	/** Tell Steam to stop capturing voice data */
+	void StopRecording() const;
+
+	/** Called when "last half second" is over */
+	void StoppedRecording() const;
+
+	/** @return is active recording occurring at the moment */
+	bool IsRecording() const { return bIsCapturing || bPendingFinalCapture; }
+
+PACKAGE_SCOPE:
+
+	/** Constructor */
+	FVoiceEngineSteam() :
+		SteamSubsystem(NULL),
+		SteamUserPtr(NULL),
+		SteamFriendsPtr(NULL),
+		OwningUserIndex(INVALID_INDEX),
+		CompressedBytesAvailable(0),
+		AvailableVoiceResult(k_EVoiceResultNotInitialized),
+		bPendingFinalCapture(false),
+		bIsCapturing(false),
+		SerializeHelper(NULL)
+	{};
+
+	// IVoiceEngine
+	virtual bool Init(int32 MaxLocalTalkers, int32 MaxRemoteTalkers) override;
+
+public:
+
+	FVoiceEngineSteam(FOnlineSubsystemSteam* InSteamSubsystem);
+	virtual ~FVoiceEngineSteam();
+
+	// IVoiceEngine
+    virtual uint32 StartLocalVoiceProcessing(uint32 LocalUserNum) override;
+    virtual uint32 StopLocalVoiceProcessing(uint32 LocalUserNum) override;
+    virtual uint32 StartRemoteVoiceProcessing(const FUniqueNetId& UniqueId) override
+	{
+		// Not needed in Steam
+		return S_OK;
+	}
+    virtual uint32 StopRemoteVoiceProcessing(const FUniqueNetId& UniqueId) override
+	{
+		// Not needed in Steam
+		return S_OK;
+	}
+
+    virtual uint32 RegisterLocalTalker(uint32 LocalUserNum) override
+	{
+		if (OwningUserIndex == INVALID_INDEX)
+		{
+			OwningUserIndex = LocalUserNum;
+			return S_OK;
+		}
+
+		return E_FAIL;
+	}
+
+    virtual uint32 UnregisterLocalTalker(uint32 LocalUserNum) override
+	{
+		if (IsOwningUser(LocalUserNum))
+		{
+			OwningUserIndex = INVALID_INDEX;
+			return S_OK;
+		}
+
+		return E_FAIL;
+	}
+
+    virtual uint32 RegisterRemoteTalker(const FUniqueNetId& UniqueId) override
+	{
+		// Not needed in Steam
+		return S_OK;
+	}
+
+    virtual uint32 UnregisterRemoteTalker(const FUniqueNetId& UniqueId) override
+	{
+		// Not needed in Steam
+		return S_OK;
+	}
+
+    virtual bool IsHeadsetPresent(uint32 LocalUserNum) override
+	{
+		return IsOwningUser(LocalUserNum) ? true : false;
+	}
+
+    virtual bool IsLocalPlayerTalking(uint32 LocalUserNum) override
+	{
+		return (GetVoiceDataReadyFlags() & (LocalUserNum << 1)) != 0;
+	}
+
+	virtual bool IsRemotePlayerTalking(const FUniqueNetId& UniqueId) override
+	{
+		return RemoteTalkerBuffers.Find((const FUniqueNetIdSteam&)UniqueId) != NULL;
+	}
+
+	virtual uint32 GetVoiceDataReadyFlags() const override;
+	virtual uint32 SetPlaybackPriority(uint32 LocalUserNum, const FUniqueNetId& RemoteTalkerId, uint32 Priority) override
+	{
+		// Not supported
+		return S_OK;
+	}
+
+	virtual uint32 ReadLocalVoiceData(uint32 LocalUserNum, uint8* Data, uint32* Size) override;
+	virtual uint32 SubmitRemoteVoiceData(const FUniqueNetId& RemoteTalkerId, uint8* Data, uint32* Size) override;
+	virtual void Tick(float DeltaTime) override;
+	FString GetVoiceDebugState() const override;
+
+	/**
+	 * Update the state of all remote talkers, possibly dropping data or the talker entirely
+	 */
+	void TickTalkers(float DeltaTime);
+
+	/**
+	 * Delegate triggered when an audio component Stop() function is called
+	 */
+	void OnAudioFinished(UAudioComponent* AC);
+};
+
+typedef TSharedPtr<FVoiceEngineSteam, ESPMode::ThreadSafe> FVoiceEngineSteamPtr;
\ No newline at end of file
diff --git a/Engine/Plugins/Online/OnlineSubsystemSteam/Source/Private/VoiceInterfaceSteam.cpp b/Engine/Plugins/Online/OnlineSubsystemSteam/Source/Private/VoiceInterfaceSteam.cpp
new file mode 100644
--- /dev/null
+++ b/Engine/Plugins/Online/OnlineSubsystemSteam/Source/Private/VoiceInterfaceSteam.cpp
@@ -0,0 +1,841 @@
+// Copyright 1998-2016 Epic Games, Inc. All Rights Reserved.
+
+#include "OnlineSubsystemSteamPrivatePCH.h"
+#include "OnlineSubsystemSteam.h"
+#include "VoiceInterfaceSteam.h"
+#include "VoiceEngineSteam.h"
+#include "Voice.h"
+#include "OnlineIdentityInterfaceSteam.h"
+#include "OnlineSessionInterfaceSteam.h"
+#include "SteamUtilities.h"
+
+/** Largest size to attempt to transmit */
+#define MAX_VOICE_PACKET_SIZE_STEAM 1 * 1024
+
+FOnlineVoiceSteam::FOnlineVoiceSteam(FOnlineSubsystemSteam* InSteamSubsystem) :
+	SteamSubsystem(InSteamSubsystem),
+	VoiceEngine(NULL)
+{
+	check(InSteamSubsystem);
+}
+
+FOnlineVoiceSteam::~FOnlineVoiceSteam()
+{
+	LocalTalkers.Empty();
+	RemoteTalkers.Empty();
+	VoiceEngine = NULL;
+}
+
+bool FOnlineVoiceSteam::Init()
+{
+	bool bSuccess = false;
+
+	if (!GConfig->GetInt(TEXT("OnlineSubsystem"),TEXT("MaxLocalTalkers"), MaxLocalTalkers, GEngineIni))
+	{
+		MaxLocalTalkers = MAX_SPLITSCREEN_TALKERS;
+		UE_LOG(LogVoice, Warning, TEXT("Missing MaxLocalTalkers key in OnlineSubsystem of DefaultEngine.ini"));
+	}
+	if (!GConfig->GetInt(TEXT("OnlineSubsystem"),TEXT("MaxRemoteTalkers"), MaxRemoteTalkers, GEngineIni))
+	{
+		MaxRemoteTalkers = MAX_REMOTE_TALKERS;
+		UE_LOG(LogVoice, Warning, TEXT("Missing MaxRemoteTalkers key in OnlineSubsystem of DefaultEngine.ini"));
+	}
+	if (!GConfig->GetFloat(TEXT("OnlineSubsystem"),TEXT("VoiceNotificationDelta"), VoiceNotificationDelta, GEngineIni))
+	{
+		VoiceNotificationDelta = 0.2;
+		UE_LOG(LogVoice, Warning, TEXT("Missing VoiceNotificationDelta key in OnlineSubsystem of DefaultEngine.ini"));
+	}
+
+	bool bHasVoiceEnabled = false;
+	if (GConfig->GetBool(TEXT("OnlineSubsystem"), TEXT("bHasVoiceEnabled"), bHasVoiceEnabled, GEngineIni) && bHasVoiceEnabled)
+	{
+		SessionInt = (FOnlineSessionSteam*)SteamSubsystem->GetSessionInterface().Get();
+		IdentityInt = (FOnlineIdentitySteam*)SteamSubsystem->GetIdentityInterface().Get();
+
+		bSuccess = SessionInt && IdentityInt;
+
+		if (bSuccess)
+		{
+			const bool bVoiceEngineForceDisable = SteamSubsystem->IsDedicated() || GIsBuildMachine;
+			if (!bVoiceEngineForceDisable)
+			{
+				VoiceEngine = MakeShareable(new FVoiceEngineSteam(SteamSubsystem));
+				bSuccess = VoiceEngine->Init(MaxLocalTalkers, MaxRemoteTalkers);
+			}
+			else
+			{
+				MaxLocalTalkers = 0;
+				MaxRemoteTalkers = 0;
+			}
+		}
+
+		LocalTalkers.Init(FLocalTalker(), MaxLocalTalkers);
+		RemoteTalkers.Empty(MaxRemoteTalkers);
+
+		if (!bSuccess)
+		{
+			UE_LOG(LogVoice, Warning, TEXT("Failed to initialize Steam voice interface"));
+
+			LocalTalkers.Empty();
+			RemoteTalkers.Empty();
+			VoiceEngine = NULL;
+		}
+	}
+	else
+	{
+		UE_LOG(LogVoice, Log, TEXT("Voice interface disabled by config [OnlineSubsystem].bHasVoiceEnabled"));
+	}
+
+	return bSuccess;
+}
+
+void FOnlineVoiceSteam::ClearVoicePackets()
+{
+	for (uint32 Index = 0; Index < MAX_SPLITSCREEN_TALKERS; Index++)
+	{
+		FVoicePacketSteam& LocalPacket = VoiceData.LocalPackets[Index];
+		if (LocalPacket.Length > 0)
+		{
+			// Mark the local packet as processed
+			LocalPacket.Length = 0;
+		}
+	}
+}
+
+void FOnlineVoiceSteam::Tick(float DeltaTime) 
+{
+	if (!SteamSubsystem->IsDedicated())
+	{
+		SCOPE_CYCLE_COUNTER(STAT_Voice_Interface);
+
+		// If we aren't in a networked match, no need to update networked voice
+		if (SessionInt && SessionInt->GetNumSessions() > 0)
+		{
+			// Processing voice data only valid with a voice engine to capture/play
+			if (VoiceEngine.IsValid())
+			{
+				VoiceEngine->Tick(DeltaTime);
+
+				// Queue local packets for sending via the network
+				ProcessLocalVoicePackets();
+				// Submit queued packets to audio system
+				ProcessRemoteVoicePackets();
+				// Fire off any talking notifications for hud display
+				ProcessTalkingDelegates(DeltaTime);
+			}
+		}
+	}
+}
+
+void FOnlineVoiceSteam::StartNetworkedVoice(uint8 LocalUserNum)
+{
+	// Validate the range of the entry
+	if (LocalUserNum < MAX_LOCAL_PLAYERS)
+	{
+		LocalTalkers[LocalUserNum].bHasNetworkedVoice = true;
+		UE_LOG(LogVoice, Log, TEXT("Starting networked voice for user: %d"), LocalUserNum);
+	}
+	else
+	{
+		UE_LOG(LogVoice, Log, TEXT("Invalid user specified in StartNetworkedVoice(%d)"),
+			(uint32)LocalUserNum);
+	}
+}
+
+void FOnlineVoiceSteam::StopNetworkedVoice(uint8 LocalUserNum)
+{
+	// Validate the range of the entry
+	if (LocalTalkers.Num() > LocalUserNum && LocalUserNum >= 0 && LocalUserNum < MAX_LOCAL_PLAYERS)
+	{
+		LocalTalkers[LocalUserNum].bHasNetworkedVoice = false;
+		UE_LOG(LogVoice, Log, TEXT("Stopping networked voice for user: %d"), LocalUserNum);
+	}
+	else
+	{
+		UE_LOG(LogVoice, Log, TEXT("Invalid user specified in StopNetworkedVoice(%d)"),
+			(uint32)LocalUserNum);
+	}
+}
+
+bool FOnlineVoiceSteam::RegisterLocalTalker(uint32 LocalUserNum) 
+{
+	uint32 Return = E_FAIL;
+	if (LocalUserNum >= 0 && LocalUserNum < MAX_LOCAL_PLAYERS)
+	{
+		// Get at the local talker's cached data
+		FLocalTalker& Talker = LocalTalkers[LocalUserNum];
+		// Make local user capable of sending voice data
+		StartNetworkedVoice(LocalUserNum);
+		// Don't register talkers when voice is disabled
+		if (VoiceEngine.IsValid())
+		{
+			if (Talker.bIsRegistered == false)
+			{
+				// Register the talker locally
+				Return = VoiceEngine->RegisterLocalTalker(LocalUserNum);
+				UE_LOG(LogVoice, Log, TEXT("RegisterLocalTalker(%d) returned 0x%08X"),	LocalUserNum, Return);
+				if (Return == S_OK)
+				{
+					Talker.bIsRegistered = true;
+					// Kick off the processing mode
+					Return = VoiceEngine->StartLocalVoiceProcessing(LocalUserNum);
+					UE_LOG(LogVoice, Log, TEXT("StartLocalProcessing(%d) returned 0x%08X"), LocalUserNum, Return);
+				}
+			}
+			else
+			{
+				// Just say yes, we registered fine
+				Return = S_OK;
+			}
+			
+			// @todo ONLINE - update mute list?
+		}
+		else
+		{
+			// Not properly logged in, so skip voice for them
+			Talker.bIsRegistered = false;
+		}
+	}
+	else
+	{
+		UE_LOG(LogVoice, Log, TEXT("Invalid user specified in RegisterLocalTalker(%d)"), LocalUserNum);
+	}
+	return Return == S_OK;
+}
+
+void FOnlineVoiceSteam::RegisterLocalTalkers()
+{
+	UE_LOG(LogVoice, Log, TEXT("Registering all local talkers"));
+	// Loop through the 4 available players and register them
+	for (uint32 Index = 0; Index < MAX_LOCAL_PLAYERS; Index++)
+	{
+		// Register the local player as a local talker
+		RegisterLocalTalker(Index);
+	}
+}
+
+bool FOnlineVoiceSteam::UnregisterLocalTalker(uint32 LocalUserNum) 
+{
+	uint32 Return = S_OK;
+	if (LocalUserNum >= 0 && LocalUserNum < MAX_LOCAL_PLAYERS)
+	{
+		// Get at the local talker's cached data
+		FLocalTalker& Talker = LocalTalkers[LocalUserNum];
+		// Skip the unregistration if not registered
+		if (Talker.bIsRegistered == true &&
+			// Or when voice is disabled
+			VoiceEngine.IsValid())
+		{
+			if (OnPlayerTalkingStateChangedDelegates.IsBound() && (Talker.bIsTalking || Talker.bWasTalking))
+			{
+				TSharedPtr<const FUniqueNetId> UniqueId = IdentityInt->GetUniquePlayerId(LocalUserNum);
+				if (UniqueId.IsValid())
+				{
+					OnPlayerTalkingStateChangedDelegates.Broadcast(UniqueId.ToSharedRef(), false);
+				}
+				else
+				{
+					UE_LOG(LogVoice, Warning, TEXT("Invalid UserId for local player %d in UnregisterLocalTalker"), LocalUserNum);
+				}
+			}
+
+			// Remove them from engine too
+			VoiceEngine->StopLocalVoiceProcessing(LocalUserNum);
+			Return = VoiceEngine->UnregisterLocalTalker(LocalUserNum);
+			UE_LOG(LogVoice, Log, TEXT("UnregisterLocalTalker(%d) returned 0x%08X"), LocalUserNum, Return);
+			Talker.bIsTalking = false;
+			Talker.bWasTalking = false;
+			Talker.bIsRegistered = false;
+		}
+	}
+	else
+	{
+		UE_LOG(LogVoice, Log, TEXT("Invalid user specified in UnregisterLocalTalker(%d)"), LocalUserNum);
+	}
+	return Return == S_OK;
+}
+
+void FOnlineVoiceSteam::UnregisterLocalTalkers()
+{
+	UE_LOG(LogVoice, Log, TEXT("Unregistering all local talkers"));
+	// Loop through the 4 available players and unregister them
+	for (uint32 Index = 0; Index < MAX_LOCAL_PLAYERS; Index++)
+	{
+		// Unregister the local player as a local talker
+		UnregisterLocalTalker(Index);
+	}
+}
+
+bool FOnlineVoiceSteam::RegisterRemoteTalker(const FUniqueNetId& UniqueId) 
+{
+	uint32 Return = E_FAIL;
+
+	// Skip this if the session isn't active
+	if (SessionInt && SessionInt->GetNumSessions() > 0 &&
+		// Or when voice is disabled
+		VoiceEngine.IsValid())
+	{
+		// See if this talker has already been registered or not
+		FRemoteTalker* Talker = FindRemoteTalker(UniqueId);
+		if (Talker == NULL)
+		{
+			// Add a new talker to our list
+			int32 AddIndex = RemoteTalkers.AddZeroed();
+			Talker = &RemoteTalkers[AddIndex];
+			// Copy the UniqueId
+			const FUniqueNetIdSteam& UniqueIdSteam = (const FUniqueNetIdSteam&)UniqueId;
+			Talker->TalkerId = MakeShareable(new FUniqueNetIdSteam(UniqueIdSteam));
+			// Register the remote talker locally
+			Return = VoiceEngine->RegisterRemoteTalker(UniqueId);
+			UE_LOG(LogVoice, Log, TEXT("RegisterRemoteTalker(%s) returned 0x%08X"),
+				*UniqueId.ToDebugString(), Return);
+		}
+		else
+		{
+			UE_LOG(LogVoice, Verbose, TEXT("Remote talker %s is being re-registered"), *UniqueId.ToDebugString());
+			Return = S_OK;
+		}
+		
+		// Update muting all of the local talkers with this remote talker
+		ProcessMuteChangeNotification();
+		// Now start processing the remote voices
+		Return = VoiceEngine->StartRemoteVoiceProcessing(UniqueId);
+		UE_LOG(LogVoice, Log, TEXT("StartRemoteVoiceProcessing(%s) returned 0x%08X"), *UniqueId.ToDebugString(), Return);
+	}
+
+	return Return == S_OK;
+}
+
+bool FOnlineVoiceSteam::UnregisterRemoteTalker(const FUniqueNetId& UniqueId) 
+{
+	uint32 Return = E_FAIL;
+
+	// Skip this if the session isn't active
+	if (SessionInt && SessionInt->GetNumSessions() > 0 &&
+		// Or when voice is disabled
+		VoiceEngine.IsValid())
+	{
+		// Make sure the talker is valid
+		if (FindRemoteTalker(UniqueId) != NULL)
+		{
+			// Find them in the talkers array and remove them
+			for (int32 Index = 0; Index < RemoteTalkers.Num(); Index++)
+			{
+				const FRemoteTalker& Talker = RemoteTalkers[Index];
+				// Is this the remote talker?
+				if (*Talker.TalkerId == UniqueId)
+				{
+					// Going to remove the talker, so if they were talking recently make sure to indicate they've stopped
+					if (OnPlayerTalkingStateChangedDelegates.IsBound() && (Talker.bIsTalking || Talker.bWasTalking))
+					{
+						OnPlayerTalkingStateChangedDelegates.Broadcast(Talker.TalkerId.ToSharedRef(), false);
+					}
+
+					RemoteTalkers.RemoveAtSwap(Index);
+					break;
+				}
+			}
+			// Remove them from voice engine
+			Return = VoiceEngine->UnregisterRemoteTalker(UniqueId);
+			UE_LOG(LogVoice, Log, TEXT("UnregisterRemoteTalker(%s) returned 0x%08X"), *UniqueId.ToDebugString(), Return);
+		}
+		else
+		{
+			UE_LOG(LogVoice, Verbose, TEXT("Unknown remote talker (%s) specified to UnregisterRemoteTalker()"), *UniqueId.ToDebugString());
+		}
+	}
+
+	return Return == S_OK;
+}
+
+void FOnlineVoiceSteam::RemoveAllRemoteTalkers()
+{
+	UE_LOG(LogVoice, Log, TEXT("Removing all remote talkers"));
+	if (VoiceEngine.IsValid())
+	{
+		// Work backwards through array removing the talkers
+		for (int32 Index = RemoteTalkers.Num() - 1; Index >= 0; Index--)
+		{
+			const FRemoteTalker& Talker = RemoteTalkers[Index];
+
+			if (OnPlayerTalkingStateChangedDelegates.IsBound() && (Talker.bIsTalking || Talker.bWasTalking))
+			{
+				OnPlayerTalkingStateChangedDelegates.Broadcast(Talker.TalkerId.ToSharedRef(), false);
+			}
+
+			uint32 Return = VoiceEngine->UnregisterRemoteTalker(*Talker.TalkerId);
+			UE_LOG(LogVoice, Log, TEXT("UnregisterRemoteTalker(%s) returned 0x%08X"), *Talker.TalkerId->ToDebugString(), Return);
+		}
+	}
+
+	// Empty the array now that they are all unregistered
+	RemoteTalkers.Empty(MaxRemoteTalkers);
+}
+
+FRemoteTalker* FOnlineVoiceSteam::FindRemoteTalker(const FUniqueNetId& UniqueId)
+{
+	for (int32 Index = 0; Index < RemoteTalkers.Num(); Index++)
+	{
+		FRemoteTalker& Talker = RemoteTalkers[Index];
+		// Compare net ids to see if they match
+		if (*Talker.TalkerId == UniqueId)
+		{
+			return &RemoteTalkers[Index];
+		}
+	}
+	return NULL;
+}
+
+bool FOnlineVoiceSteam::IsHeadsetPresent(uint32 LocalUserNum) 
+{
+	return VoiceEngine.IsValid() && VoiceEngine->IsHeadsetPresent(LocalUserNum);
+}
+
+bool FOnlineVoiceSteam::IsLocalPlayerTalking(uint32 LocalUserNum) 
+{
+	return VoiceEngine.IsValid() && VoiceEngine->IsLocalPlayerTalking(LocalUserNum);
+}
+
+bool FOnlineVoiceSteam::IsRemotePlayerTalking(const FUniqueNetId& UniqueId) 
+{
+	return VoiceEngine.IsValid() && VoiceEngine->IsRemotePlayerTalking(UniqueId);
+}
+
+bool FOnlineVoiceSteam::IsMuted(uint32 LocalUserNum, const FUniqueNetId& UniqueId) const
+{
+	int32 Index = INDEX_NONE;
+	if (LocalUserNum >= 0 && LocalUserNum < MAX_LOCAL_PLAYERS)
+	{
+		Index = MuteList.Find((const FUniqueNetIdSteam&)UniqueId);
+	}
+
+	return Index != INDEX_NONE;
+}
+
+bool FOnlineVoiceSteam::IsLocallyMuted(const FUniqueNetId& UniqueId) const
+{
+	int32 Index = MuteList.Find((const FUniqueNetIdSteam&)UniqueId);
+	return Index != INDEX_NONE;
+}
+
+bool FOnlineVoiceSteam::MuteRemoteTalker(uint8 LocalUserNum, const FUniqueNetId& PlayerId, bool bIsSystemWide)
+{
+	uint32 Return = E_FAIL;
+	if (LocalUserNum >= 0 && LocalUserNum < MAX_LOCAL_PLAYERS)
+	{
+		if (bIsSystemWide)
+		{
+			SystemMuteList.AddUnique((const FUniqueNetIdSteam&)PlayerId);
+			ProcessMuteChangeNotification();
+		}
+		else
+		{
+			// Skip this if the session isn't active
+			if (SessionInt && SessionInt->GetNumSessions() > 0 &&
+				// Or if voice is disabled
+					VoiceEngine.IsValid())
+			{
+				// Find the specified talker
+				FRemoteTalker* Talker = FindRemoteTalker(PlayerId);
+				if (Talker != NULL)
+				{
+					MuteList.AddUnique((const FUniqueNetIdSteam&)PlayerId);
+					Return = S_OK;
+					UE_LOG(LogVoice, Log, TEXT("Muting remote talker (%s)"), *PlayerId.ToDebugString());
+				}
+				else
+				{
+					UE_LOG(LogVoice, Verbose, TEXT("Unknown remote talker (%s) specified to MuteRemoteTalker()"), *PlayerId.ToDebugString());
+				}
+			}
+		}
+	}
+	else
+	{
+		UE_LOG(LogVoice, Warning, TEXT("Invalid user specified in MuteRemoteTalker(%d)"), LocalUserNum);
+	}
+
+	return Return == S_OK;
+}
+
+bool FOnlineVoiceSteam::UnmuteRemoteTalker(uint8 LocalUserNum, const FUniqueNetId& PlayerId, bool bIsSystemWide)
+{
+	uint32 Return = E_FAIL;
+	if (LocalUserNum >= 0 && LocalUserNum < MAX_LOCAL_PLAYERS)
+	{
+		if (bIsSystemWide)
+		{
+			// Remove them from the mute list
+			SystemMuteList.RemoveSingleSwap((const FUniqueNetIdSteam&)PlayerId);
+			ProcessMuteChangeNotification();
+		}
+		else
+		{
+			// Skip this if the session isn't active
+			if (SessionInt && SessionInt->GetNumSessions() > 0 &&
+				// Or if voice is disabled
+				VoiceEngine.IsValid())
+			{
+				// Find the specified talker
+				FRemoteTalker* Talker = FindRemoteTalker(PlayerId);
+				if (Talker != NULL)
+				{
+					// Make sure there isn't a system mute
+					bool bIsMuted = IsMuted(LocalUserNum, PlayerId);
+					if (!bIsMuted)
+					{
+						// Remove them from the mute list
+						MuteList.RemoveSingleSwap((const FUniqueNetIdSteam&)PlayerId);
+						UE_LOG(LogVoice, Log, TEXT("Unmuting remote talker (%s)"), *PlayerId.ToDebugString());
+					}
+				}
+				else
+				{
+					UE_LOG(LogVoice, Verbose, TEXT("Unknown remote talker (%s) specified to UnmuteRemoteTalker()"), *PlayerId.ToDebugString());
+				}
+			}
+		}
+	}
+	else
+	{
+		UE_LOG(LogVoice, Warning, TEXT("Invalid user specified in UnmuteRemoteTalker(%d)"), LocalUserNum);
+	}
+
+	return Return == S_OK;
+}
+
+void FOnlineVoiceSteam::ProcessMuteChangeNotification()
+{
+	// Nothing to update if there isn't an active session
+	if (VoiceEngine.IsValid())
+	{
+		if (SessionInt && SessionInt->GetNumSessions() > 0)
+		{
+			// For each local user with voice
+			for (int32 Index = 0; Index < MaxLocalTalkers; Index++)
+			{
+				// Find the very first ULocalPlayer for this ControllerId. 
+				// This is imperfect and means we cannot support voice chat properly for
+				// multiple UWorlds (but thats ok for the time being).
+				ULocalPlayer* LP = GEngine->FindFirstLocalPlayerFromControllerId(Index);
+				if (LP && LP->PlayerController)
+				{
+					// If there is a player controller, we can mute/unmute people
+					if (LocalTalkers[Index].bIsRegistered && LP->PlayerController != NULL)
+					{
+						// Use the common method of checking muting
+						UpdateMuteListForLocalTalker(Index, LP->PlayerController);
+					}
+				}
+			}
+		}
+	}
+}
+
+void FOnlineVoiceSteam::UpdateMuteListForLocalTalker(int32 TalkerIndex, APlayerController* PlayerController)
+{
+	// For each registered remote talker
+	for (int32 RemoteIndex = 0; RemoteIndex < RemoteTalkers.Num(); RemoteIndex++)
+	{
+		const FRemoteTalker& Talker = RemoteTalkers[RemoteIndex];
+
+		FUniqueNetIdRepl UniqueIdRepl(Talker.TalkerId);
+
+		// Is the remote talker on this local player's mute list?
+		if (SystemMuteList.Find((FUniqueNetIdSteam&)*Talker.TalkerId) == INDEX_NONE)
+		{
+			// Unmute on the server
+			PlayerController->ServerUnmutePlayer(UniqueIdRepl);
+		}
+		else
+		{
+			// Mute on the server
+			PlayerController->ServerMutePlayer(UniqueIdRepl);
+		}
+
+		// The ServerUn/MutePlayer() functions will perform the muting based
+		// upon gameplay settings and other player's mute list
+	}
+}
+
+TSharedPtr<FVoicePacket> FOnlineVoiceSteam::SerializeRemotePacket(FArchive& Ar)
+{
+	TSharedPtr<FVoicePacketSteam> NewPacket = MakeShareable(new FVoicePacketSteam());
+	NewPacket->Serialize(Ar);
+	if (Ar.IsError() == false && NewPacket->GetBufferSize() > 0)
+	{
+		if (!SteamSubsystem->IsDedicated())
+		{
+			FUniqueNetIdMatcher PlayerMatch(*NewPacket->GetSender());
+			if (MuteList.IndexOfByPredicate(PlayerMatch) == INDEX_NONE)
+			{
+				VoiceData.RemotePackets.Add(NewPacket);
+			}
+		}
+
+		return NewPacket;
+	}
+
+	return NULL;
+}
+
+TSharedPtr<FVoicePacket> FOnlineVoiceSteam::GetLocalPacket(uint32 LocalUserNum)
+{
+	// duplicate the local copy of the data and set it on a shared pointer for destruction elsewhere
+	if (LocalUserNum >= 0 && LocalUserNum < MAX_SPLITSCREEN_TALKERS)
+	{
+		FVoicePacketSteam& VoicePacket = VoiceData.LocalPackets[LocalUserNum];
+		if (VoicePacket.GetBufferSize() > 0)
+		{
+			return MakeShareable(new FVoicePacketSteam(VoicePacket));
+		}
+	}
+
+	return NULL;
+}
+
+void FOnlineVoiceSteam::ProcessTalkingDelegates(float DeltaTime)
+{
+	// Fire off any talker notification delegates for local talkers
+	for (int32 LocalUserNum = 0; LocalUserNum < LocalTalkers.Num(); LocalUserNum++)
+	{
+		FLocalTalker& Talker = LocalTalkers[LocalUserNum];
+
+		// Only check players with voice
+		if (Talker.bIsRegistered)
+		{
+			// If the talker was not previously talking, but now is trigger the event
+			bool bShouldNotify = !Talker.bWasTalking && Talker.bIsTalking;
+			// If the talker was previously talking, but now isn't time delay the event
+			if (!bShouldNotify && Talker.bWasTalking)
+			{
+				Talker.LastNotificationTime -= DeltaTime;
+				if (Talker.LastNotificationTime <= 0.f)
+				{
+					// Clear the flag so it only activates when needed
+					Talker.bIsTalking = false;
+					Talker.LastNotificationTime = VoiceNotificationDelta;
+					bShouldNotify = true;
+				}
+			}
+
+			if (bShouldNotify)
+			{
+				// Skip all delegate handling if none are registered
+				if (OnPlayerTalkingStateChangedDelegates.IsBound())
+				{
+					TSharedPtr<const FUniqueNetId> UniqueId = IdentityInt->GetUniquePlayerId(LocalUserNum);
+					OnPlayerTalkingStateChangedDelegates.Broadcast(UniqueId.ToSharedRef(), Talker.bIsTalking);
+				}
+
+				Talker.bWasTalking = Talker.bIsTalking;
+				UE_LOG(LogVoice, Log, TEXT("Trigger Local %d %sTALKING"), LocalUserNum, Talker.bIsTalking ? TEXT("") : TEXT("NOT"));
+			}
+		}
+	}
+	// Now check all remote talkers
+	for (int32 Index = 0; Index < RemoteTalkers.Num(); Index++)
+	{
+		FRemoteTalker& Talker = RemoteTalkers[Index];
+
+		// If the talker was not previously talking, but now is trigger the event
+		bool bShouldNotify = !Talker.bWasTalking && Talker.bIsTalking;
+		// If the talker was previously talking, but now isn't time delay the event
+		if (!bShouldNotify && Talker.bWasTalking && !Talker.bIsTalking)
+		{
+			Talker.LastNotificationTime -= DeltaTime;
+			if (Talker.LastNotificationTime <= 0.f)
+			{
+				bShouldNotify = true;
+			}
+		}
+
+		if (bShouldNotify)
+		{
+			// Skip all delegate handling if none are registered
+			if (OnPlayerTalkingStateChangedDelegates.IsBound())
+			{
+				OnPlayerTalkingStateChangedDelegates.Broadcast(Talker.TalkerId.ToSharedRef(), Talker.bIsTalking);
+			}
+
+			UE_LOG(LogVoice, Log, TEXT("Trigger Remote %s %sTALKING"), *Talker.TalkerId->ToString(), Talker.bIsTalking ? TEXT("") : TEXT("NOT"));
+
+			// Clear the flag so it only activates when needed
+			Talker.bWasTalking = Talker.bIsTalking;
+			Talker.LastNotificationTime = VoiceNotificationDelta;
+		}
+	}
+}
+
+void FOnlineVoiceSteam::ProcessLocalVoicePackets()
+{
+	if (VoiceEngine.IsValid())
+	{
+		// Read the data from any local talkers
+		uint32 DataReadyFlags = VoiceEngine->GetVoiceDataReadyFlags();
+		// Skip processing if there is no data from a local talker
+		if (DataReadyFlags)
+		{
+			// Process each talker with a bit set
+			for (uint32 Index = 0; DataReadyFlags; Index++, DataReadyFlags >>= 1)
+			{
+				// Talkers needing processing will always be in lsb due to shifts
+				if (DataReadyFlags & 1)
+				{
+					uint32 SpaceAvail = MAX_VOICE_DATA_SIZE - VoiceData.LocalPackets[Index].Length;
+					// Figure out if there is space for this packet
+					if (SpaceAvail > 0)
+					{
+						// Figure out where to append the data
+						uint8* BufferStart = VoiceData.LocalPackets[Index].Buffer.GetData();
+						BufferStart += VoiceData.LocalPackets[Index].Length;
+						// Copy the sender info
+						VoiceData.LocalPackets[Index].Sender = IdentityInt->GetUniquePlayerId(Index);
+						// Process this user
+						uint32 Result = VoiceEngine->ReadLocalVoiceData(Index, BufferStart, &SpaceAvail);
+						if (Result == S_OK)
+						{
+							if (LocalTalkers[Index].bHasNetworkedVoice && SpaceAvail <= MAX_VOICE_PACKET_SIZE_STEAM)
+							{
+								// Mark the person as talking
+								LocalTalkers[Index].bIsTalking = true;
+								LocalTalkers[Index].LastNotificationTime = VoiceNotificationDelta;
+
+								// Update the length based on what it copied
+								VoiceData.LocalPackets[Index].Length += SpaceAvail;
+
+#if VOICE_LOOPBACK
+								static bool bUseLoopback = false;
+								if (bUseLoopback && SpaceAvail > 0)
+								{
+									VoiceData.RemotePackets.Add(MakeShareable(new FVoicePacketSteam(VoiceData.LocalPackets[Index])));
+								}
+#endif
+							}
+							else
+							{
+								// Zero out the data since it isn't to be sent via the network
+								VoiceData.LocalPackets[Index].Length = 0;
+							}
+						}
+						else
+						{
+							UE_LOG(LogVoiceEncode, Warning, TEXT("Voice data error in ReadLocalVoiceData"));
+						}
+					}
+					else
+					{
+						UE_LOG(LogVoiceEncode, Warning, TEXT("Dropping voice data due to network layer not processing fast enough"));
+						// Buffer overflow, so drop previous data
+						VoiceData.LocalPackets[Index].Length = 0;
+					}
+				}
+			}
+		}
+	}
+}
+
+void FOnlineVoiceSteam::ProcessRemoteVoicePackets()
+{
+	// Clear the talking state for remote players
+	for (int32 Index = 0; Index < RemoteTalkers.Num(); Index++)
+	{
+		RemoteTalkers[Index].bIsTalking = false;
+	}
+
+	// Now process all pending packets from the server
+	for (int32 Index = 0; Index < VoiceData.RemotePackets.Num(); Index++)
+	{
+		TSharedPtr<FVoicePacketSteam> VoicePacket = StaticCastSharedPtr<FVoicePacketSteam>(VoiceData.RemotePackets[Index]);
+		if (VoicePacket.IsValid())
+		{
+			// Skip local submission of voice if dedicated server or no voice
+			if (VoiceEngine.IsValid())
+			{
+				// Get the size since it is an in/out param
+				uint32 VoiceBufferSize = VoicePacket->GetBufferSize();
+				// Submit this packet to the voice engine
+				uint32 Result = VoiceEngine->SubmitRemoteVoiceData(*VoicePacket->Sender, VoicePacket->Buffer.GetData(), &VoiceBufferSize);
+				if (Result != S_OK)
+				{
+					UE_LOG(LogVoiceDecode, Log,
+						TEXT("SubmitRemoteVoiceData(%s) failed with 0x%08X"),
+						*VoicePacket->Sender->ToDebugString(),
+						Result);
+				}
+			}
+
+			// Find the remote talker and mark them as talking
+			for (int32 Index2 = 0; Index2 < RemoteTalkers.Num(); Index2++)
+			{
+				FRemoteTalker& Talker = RemoteTalkers[Index2];
+				// Compare the ids
+				if (*Talker.TalkerId == *VoicePacket->Sender)
+				{
+					// If the player is marked as muted, they can't be talking
+					Talker.bIsTalking = !IsLocallyMuted(*Talker.TalkerId);
+					Talker.LastNotificationTime = VoiceNotificationDelta;
+				}
+			}
+		}
+	}
+	// Zero the list without causing a free/realloc
+	VoiceData.RemotePackets.Reset();
+}
+
+FString FOnlineVoiceSteam::GetVoiceDebugState() const
+{
+	TSharedPtr<const FUniqueNetId> UniqueId;
+
+	FString Output = TEXT("Voice state\n");
+	Output += VoiceEngine.IsValid() ? VoiceEngine->GetVoiceDebugState() : TEXT("No Voice Engine!");
+
+	Output += TEXT("\nLocal Talkers:\n");
+	for (int32 idx=0; idx < LocalTalkers.Num(); idx++)
+	{
+		UniqueId = IdentityInt->GetUniquePlayerId(idx);
+		
+		const FLocalTalker& Talker = LocalTalkers[idx];
+		Output += FString::Printf(TEXT("[%d]: %s\n Registered: %d\n Networked: %d\n Talking: %d\n WasTalking: %d\n Last:%0.2f\n"),
+			idx,
+			UniqueId.IsValid() ? *UniqueId->ToDebugString() : TEXT("NULL"), 
+			Talker.bIsRegistered,
+			Talker.bHasNetworkedVoice,
+			Talker.bIsTalking,
+			Talker.bWasTalking,
+			Talker.LastNotificationTime);
+	}
+
+	Output += TEXT("\nRemote Talkers:\n");
+	for (int32 idx=0; idx < RemoteTalkers.Num(); idx++)
+	{
+		const FRemoteTalker& Talker = RemoteTalkers[idx];
+		Output += FString::Printf(TEXT("[%d]: %s\n Talking: %d\n WasTalking: %d\n Muted: %s\n Last:%0.2f\n"),
+			idx,
+			*Talker.TalkerId->ToDebugString(), 
+			Talker.bIsTalking,
+			Talker.bWasTalking,
+			IsLocallyMuted(*Talker.TalkerId) ? TEXT("1") : TEXT("0"),
+			Talker.LastNotificationTime);
+
+	}
+
+	Output += TEXT("\nRaw SystemMutelist:\n");
+	for (int32 idx=0; idx < SystemMuteList.Num(); idx++)
+	{
+		Output += FString::Printf(TEXT("[%d]=%s\n"),
+			idx,
+			*SystemMuteList[idx].ToString());
+	}
+
+	Output += TEXT("\nRaw Mutelist:\n");
+	for (int32 idx=0; idx < MuteList.Num(); idx++)
+	{
+		Output += FString::Printf(TEXT("[%d]=%s\n"),
+			idx,
+			*MuteList[idx].ToString());
+	}
+
+	return Output;
+}
+
diff --git a/Engine/Plugins/Online/OnlineSubsystemSteam/Source/Private/VoiceInterfaceSteam.h b/Engine/Plugins/Online/OnlineSubsystemSteam/Source/Private/VoiceInterfaceSteam.h
new file mode 100644
--- /dev/null
+++ b/Engine/Plugins/Online/OnlineSubsystemSteam/Source/Private/VoiceInterfaceSteam.h
@@ -0,0 +1,131 @@
+// Copyright 1998-2016 Epic Games, Inc. All Rights Reserved.
+
+#pragma once
+
+#include "VoiceInterface.h"
+#include "VoicePacketSteam.h"
+#include "OnlineSubsystemSteamTypes.h"
+#include "OnlineSubsystemSteamPackage.h"
+
+/**
+ * The Steam implementation of the voice interface 
+ */
+class FOnlineVoiceSteam : public IOnlineVoice
+{
+	/** Reference to the main Steam subsystem */
+	class FOnlineSubsystemSteam* SteamSubsystem;
+	/** Reference to the sessions interface */
+	class FOnlineSessionSteam* SessionInt;
+	/** Reference to the profile interface */
+	class FOnlineIdentitySteam* IdentityInt;
+	/** Reference to the voice engine for acquiring voice data */
+	IVoiceEnginePtr VoiceEngine;
+
+	/** Maximum permitted local talkers */
+	int32 MaxLocalTalkers;
+	/** Maximum permitted remote talkers */
+	int32 MaxRemoteTalkers;
+
+	/** State of all possible local talkers */
+	TArray<FLocalTalker> LocalTalkers;
+	/** State of all possible remote talkers */
+	TArray<FRemoteTalker> RemoteTalkers;
+	/** Remote players locally muted explicitly */
+	TArray<FUniqueNetIdSteam> SystemMuteList;
+	/** Remote players locally muted */
+	TArray<FUniqueNetIdSteam> MuteList;
+
+	/** Time to wait for new data before triggering "not talking" */
+	float VoiceNotificationDelta;
+
+	/** Buffered voice data I/O */
+	FVoiceDataSteam VoiceData;
+
+	/**
+	 * Finds a remote talker in the cached list
+	 *
+	 * @param UniqueId the net id of the player to search for
+	 *
+	 * @return pointer to the remote talker or NULL if not found
+	 */
+	struct FRemoteTalker* FindRemoteTalker(const FUniqueNetId& UniqueId);
+
+	/**
+	 * Is a given id presently muted (either by system mute or game server)
+	 */
+	bool IsLocallyMuted(const FUniqueNetId& UniqueId) const;
+
+PACKAGE_SCOPE:
+	FOnlineVoiceSteam() :
+		SteamSubsystem(NULL),
+		SessionInt(NULL),
+		IdentityInt(NULL),
+		VoiceEngine(NULL),
+		MaxLocalTalkers(MAX_SPLITSCREEN_TALKERS),
+		MaxRemoteTalkers(MAX_REMOTE_TALKERS),
+		VoiceNotificationDelta(0.0f)
+	{};
+
+	// IOnlineVoice
+	virtual bool Init() override;
+	void ProcessMuteChangeNotification() override;
+
+	/**
+	 * Processes any talking delegates that need to be fired off
+	 *
+	 * @param DeltaTime the amount of time that has elapsed since the last tick
+	 */
+	void ProcessTalkingDelegates(float DeltaTime);
+
+	/**
+	 * Reads any data that is currently queued
+	 */
+	void ProcessLocalVoicePackets();
+
+	/**
+	 * Submits network packets to audio system for playback
+	 */
+	void ProcessRemoteVoicePackets();
+
+	/**
+	 * Figures out which remote talkers need to be muted for a given local talker
+	 *
+	 * @param TalkerIndex the talker that needs the mute list checked for
+	 * @param PlayerController the player controller associated with this talker
+	 */
+	void UpdateMuteListForLocalTalker(int32 TalkerIndex, class APlayerController* PlayerController);
+
+public:
+
+	/** Constructor */
+	FOnlineVoiceSteam(class FOnlineSubsystemSteam* InSteamSubsystem);
+
+	/** Virtual destructor to force proper child cleanup */
+	virtual ~FOnlineVoiceSteam();
+
+	// IOnlineVoice
+	virtual void StartNetworkedVoice(uint8 LocalUserNum) override;
+	virtual void StopNetworkedVoice(uint8 LocalUserNum) override;
+    virtual bool RegisterLocalTalker(uint32 LocalUserNum) override;
+	virtual void RegisterLocalTalkers() override;
+    virtual bool UnregisterLocalTalker(uint32 LocalUserNum) override;
+	virtual void UnregisterLocalTalkers() override;
+    virtual bool RegisterRemoteTalker(const FUniqueNetId& UniqueId) override;
+    virtual bool UnregisterRemoteTalker(const FUniqueNetId& UniqueId) override;
+	virtual void RemoveAllRemoteTalkers() override;
+    virtual bool IsHeadsetPresent(uint32 LocalUserNum) override;
+    virtual bool IsLocalPlayerTalking(uint32 LocalUserNum) override;
+	virtual bool IsRemotePlayerTalking(const FUniqueNetId& UniqueId) override;
+	bool IsMuted(uint32 LocalUserNum, const FUniqueNetId& UniqueId) const override;
+	bool MuteRemoteTalker(uint8 LocalUserNum, const FUniqueNetId& PlayerId, bool bIsSystemWide) override;
+	bool UnmuteRemoteTalker(uint8 LocalUserNum, const FUniqueNetId& PlayerId, bool bIsSystemWide) override;
+	virtual TSharedPtr<class FVoicePacket> SerializeRemotePacket(FArchive& Ar) override;
+	virtual TSharedPtr<class FVoicePacket> GetLocalPacket(uint32 LocalUserNum) override;
+	virtual int32 GetNumLocalTalkers() override { return LocalTalkers.Num(); };
+	virtual void ClearVoicePackets() override;
+	virtual void Tick(float DeltaTime) override;
+	virtual FString GetVoiceDebugState() const override;
+};
+
+typedef TSharedPtr<FOnlineVoiceSteam, ESPMode::ThreadSafe> FOnlineVoiceSteamPtr;
+
diff --git a/Engine/Plugins/Online/OnlineSubsystemSteam/Source/Public/OnlineSubsystemSteam.h b/Engine/Plugins/Online/OnlineSubsystemSteam/Source/Public/OnlineSubsystemSteam.h
new file mode 100644
--- /dev/null
+++ b/Engine/Plugins/Online/OnlineSubsystemSteam/Source/Public/OnlineSubsystemSteam.h
@@ -0,0 +1,302 @@
+// Copyright 1998-2016 Epic Games, Inc. All Rights Reserved.
+
+#pragma once
+
+#include "OnlineSubsystem.h"
+#include "OnlineSubsystemImpl.h"
+#include "OnlineSubsystemSteamPackage.h"
+
+/** Forward declarations of all interface classes */
+typedef TSharedPtr<class FOnlineSessionSteam, ESPMode::ThreadSafe> FOnlineSessionSteamPtr;
+typedef TSharedPtr<class FOnlineIdentitySteam, ESPMode::ThreadSafe> FOnlineIdentitySteamPtr;
+typedef TSharedPtr<class FOnlineFriendsSteam, ESPMode::ThreadSafe> FOnlineFriendsSteamPtr;
+typedef TSharedPtr<class FOnlineSharedCloudSteam, ESPMode::ThreadSafe> FOnlineSharedCloudSteamPtr;
+typedef TSharedPtr<class FOnlineUserCloudSteam, ESPMode::ThreadSafe> FOnlineUserCloudSteamPtr;
+typedef TSharedPtr<class FOnlineLeaderboardsSteam, ESPMode::ThreadSafe> FOnlineLeaderboardsSteamPtr;
+typedef TSharedPtr<class FOnlineVoiceSteam, ESPMode::ThreadSafe> FOnlineVoiceSteamPtr;
+typedef TSharedPtr<class FOnlineExternalUISteam, ESPMode::ThreadSafe> FOnlineExternalUISteamPtr;
+typedef TSharedPtr<class FOnlineAchievementsSteam, ESPMode::ThreadSafe> FOnlineAchievementsSteamPtr;
+
+/**
+ *	OnlineSubsystemSteam - Implementation of the online subsystem for STEAM services
+ */
+class ONLINESUBSYSTEMSTEAM_API FOnlineSubsystemSteam : 
+	public FOnlineSubsystemImpl
+{
+protected:
+
+	/** Has the STEAM client APIs been initialized */
+	bool bSteamworksClientInitialized;
+
+	/** Whether or not the Steam game server API is initialized */
+	bool bSteamworksGameServerInitialized;
+
+	/** Steam App ID for the running game */
+	uint32 SteamAppID;
+
+	/** Steam port - the local port used to communicate with the steam servers */
+	int32 GameServerSteamPort;
+
+	/** Game port - the port that clients will connect to for gameplay */
+	int32 GameServerGamePort;
+
+	/** Query port - the port that will manage server browser related duties and info */
+	int32 GameServerQueryPort;
+
+	/** Array of the files in the cloud for a given user */
+	TArray<struct FSteamUserCloudData*> UserCloudData;
+
+	/** Interface to the session services */
+	FOnlineSessionSteamPtr SessionInterface;
+
+	/** Interface to the profile services */
+	FOnlineIdentitySteamPtr IdentityInterface;
+
+	/** Interface to the friend services */
+	FOnlineFriendsSteamPtr FriendInterface;
+
+	/** Interface to the shared cloud services */
+	FOnlineSharedCloudSteamPtr SharedCloudInterface;
+	
+	/** Interface to the user cloud services */
+	FOnlineUserCloudSteamPtr UserCloudInterface;
+
+	/** Interface to the leaderboard services */
+	FOnlineLeaderboardsSteamPtr LeaderboardsInterface;
+
+	/** Interface to the voice engine */
+	FOnlineVoiceSteamPtr VoiceInterface;
+
+	/** Interface to the external UI services */
+	FOnlineExternalUISteamPtr ExternalUIInterface;
+
+	/** Interface for achievements */
+	FOnlineAchievementsSteamPtr AchievementsInterface;
+
+	/** Online async task runnable */
+	class FOnlineAsyncTaskManagerSteam* OnlineAsyncTaskThreadRunnable;
+
+	/** Online async task thread */
+	class FRunnableThread* OnlineAsyncTaskThread;
+
+PACKAGE_SCOPE:
+
+	/** Only the factory makes instances */
+	FOnlineSubsystemSteam(FName InInstanceName) :
+		FOnlineSubsystemImpl(InInstanceName),
+		bSteamworksClientInitialized(false),
+		bSteamworksGameServerInitialized(false),
+		SteamAppID(0),
+		GameServerSteamPort(0),
+		GameServerGamePort(0),
+		GameServerQueryPort(0),
+		SessionInterface(NULL),
+		IdentityInterface(NULL),
+		FriendInterface(NULL),
+		SharedCloudInterface(NULL),
+		UserCloudInterface(NULL),
+		LeaderboardsInterface(NULL),
+		VoiceInterface(NULL),
+		ExternalUIInterface(NULL),
+		OnlineAsyncTaskThreadRunnable(NULL),
+		OnlineAsyncTaskThread(NULL)
+	{}
+
+	FOnlineSubsystemSteam() : 
+		bSteamworksClientInitialized(false),
+		bSteamworksGameServerInitialized(false),
+		SteamAppID(0),
+		GameServerSteamPort(0),
+		GameServerGamePort(0),
+		GameServerQueryPort(0),
+		SessionInterface(NULL),
+		IdentityInterface(NULL),
+		FriendInterface(NULL),
+		SharedCloudInterface(NULL),
+		UserCloudInterface(NULL),
+		LeaderboardsInterface(NULL),
+		VoiceInterface(NULL),
+		ExternalUIInterface(NULL),
+		OnlineAsyncTaskThreadRunnable(NULL),
+		OnlineAsyncTaskThread(NULL)
+	{}
+
+	/** Critical sections for thread safe operation of the cloud files */
+	FCriticalSection UserCloudDataLock;
+
+	/**
+	 * Is Steam available for use
+	 * @return true if Steam is available, false otherwise
+	 */
+	bool IsEnabled();
+
+	/**
+	 *	Initialize the client side APIs for Steam 
+	 * @return true if the API was initialized successfully, false otherwise
+	 */
+	bool InitSteamworksClient(bool bRelaunchInSteam, int32 SteamAppId);
+
+	/**
+	 *	Initialize the server side APIs for Steam 
+	 * @return true if the API was initialized successfully, false otherwise
+	 */
+	bool InitSteamworksServer();
+
+	/**
+	 *	Shutdown the Steam APIs
+	 */
+	void ShutdownSteamworks();
+	
+	/**
+	 *	Add an async task onto the task queue for processing
+	 * @param AsyncTask - new heap allocated task to process on the async task thread
+	 */
+	void QueueAsyncTask(class FOnlineAsyncTask* AsyncTask);
+
+	/**
+	 *	Add an async task onto the outgoing task queue for processing
+	 * @param AsyncItem - new heap allocated task to process on the async task thread
+	 */
+	void QueueAsyncOutgoingItem(class FOnlineAsyncItem* AsyncItem);
+
+    /** 
+     * **INTERNAL**
+     * Get the metadata related to a given user
+     * This information is only available after calling EnumerateUserFiles
+     *
+     * @param UserId the UserId to search for
+     * @return the struct with the metadata about the requested user, will always return a valid struct, creating one if necessary
+     *
+     */
+	struct FSteamUserCloudData* GetUserCloudEntry(const FUniqueNetId& UserId);
+
+	/** 
+     * **INTERNAL**
+     * Clear the metadata related to a given user's file on Steam
+     * This information is only available after calling EnumerateUserFiles
+     * It doesn't actually delete any of the actual data on disk
+     *
+     * @param UserId the UserId for the file to search for
+     * @param Filename the file to get metadata about
+     * @return the true if the delete was successful, false otherwise
+     *
+     */
+	bool ClearUserCloudMetadata(const FUniqueNetId& UserId, const FString& Filename);
+
+	/**
+	 *	Clear out all the data related to user cloud storage
+	 */
+	void ClearUserCloudFiles();
+
+	/** 
+	 * **INTERNAL** 
+	 * Get the interface for accessing leaderboards/stats
+	 *
+	 * @return pointer for the appropriate class
+	 */
+	FOnlineLeaderboardsSteam * GetInternalLeaderboardsInterface();
+
+public:
+
+	virtual ~FOnlineSubsystemSteam()
+	{
+	}
+
+	// IOnlineSubsystem
+
+	virtual IOnlineSessionPtr GetSessionInterface() const override;
+	virtual IOnlineFriendsPtr GetFriendsInterface() const override;
+	virtual IOnlinePartyPtr GetPartyInterface() const override;
+	virtual IOnlineGroupsPtr GetGroupsInterface() const override;
+	virtual IOnlineSharedCloudPtr GetSharedCloudInterface() const override;
+	virtual IOnlineUserCloudPtr GetUserCloudInterface() const override;
+	virtual IOnlineLeaderboardsPtr GetLeaderboardsInterface() const override;
+	virtual IOnlineVoicePtr GetVoiceInterface() const override;
+	virtual IOnlineExternalUIPtr GetExternalUIInterface() const override;
+	virtual IOnlineTimePtr GetTimeInterface() const override;
+	virtual IOnlineIdentityPtr GetIdentityInterface() const override;
+	virtual IOnlineTitleFilePtr GetTitleFileInterface() const override;
+	virtual IOnlineEntitlementsPtr GetEntitlementsInterface() const override;
+	virtual IOnlineStorePtr GetStoreInterface() const override;
+	virtual IOnlineStoreV2Ptr GetStoreV2Interface() const override { return nullptr; }
+	virtual IOnlinePurchasePtr GetPurchaseInterface() const override { return nullptr; }
+	virtual IOnlineEventsPtr GetEventsInterface() const override;
+	virtual IOnlineAchievementsPtr GetAchievementsInterface() const override;
+	virtual IOnlineSharingPtr GetSharingInterface() const override;
+	virtual IOnlineUserPtr GetUserInterface() const override;
+	virtual IOnlineMessagePtr GetMessageInterface() const override;
+	virtual IOnlinePresencePtr GetPresenceInterface() const override;
+	virtual IOnlineChatPtr GetChatInterface() const override;
+	virtual IOnlineTurnBasedPtr GetTurnBasedInterface() const override;
+	virtual bool IsLocalPlayer(const FUniqueNetId& UniqueId) const override;
+	virtual bool Init() override;
+	virtual bool Shutdown() override;
+	virtual bool Exec(class UWorld* InWorld, const TCHAR* Cmd, FOutputDevice& Ar) override;
+	FString GetAppId() const override;
+
+	// FTickerObjectBase
+
+	virtual bool Tick(float DeltaTime) override;
+
+	// FOnlineSubsystemSteam
+
+	/**
+	 * Whether or not the Steam Client interfaces are available; these interfaces are only available, if the Steam Client program is running
+	 * NOTE: These interfaces are made unavailable, when running a dedicated server
+	 *
+	 * @return	Whether or not the Steam Client interfaces are available
+	 */
+	inline bool IsSteamClientAvailable()
+	{
+		return bSteamworksClientInitialized;
+	}
+
+	/**
+	 * Whether or not the Steam game server interfaces are available; these interfaces are always available, so long as they were initialized correctly
+	 * NOTE: The Steam Client does not need to be running for the game server interfaces to initialize
+	 * NOTE: These interfaces are made unavailable, when not running a server
+	 *
+	 * @return	Whether or not the Steam game server interfaces are available
+	 */
+	inline bool IsSteamServerAvailable()
+	{
+		// @todo Steam - add some logic to detect somehow we intended to be a "Steam client" but failed that part
+		// yet managed to still initialize the game server aspects of Steam
+		return bSteamworksGameServerInitialized;
+	}
+
+	/**
+	 * @return the steam app id for this app
+	 */
+	inline uint32 GetSteamAppId() const
+	{
+		return SteamAppID;
+	}
+
+	/**
+	 *	@return the port the game has registered for play 
+	 */
+	inline int32 GetGameServerGamePort() const
+	{
+		return GameServerGamePort;
+	}
+
+	/**
+	 *	@return the port the game has registered for talking to Steam
+	 */
+	inline int32 GetGameServerSteamPort() const
+	{
+		return GameServerSteamPort;
+	}
+
+	/**
+	 *	@return the port the game has registered for incoming server queries
+	 */
+	inline int32 GetGameServerQueryPort() const
+	{
+		return GameServerQueryPort;
+	}
+};
+
+typedef TSharedPtr<FOnlineSubsystemSteam, ESPMode::ThreadSafe> FOnlineSubsystemSteamPtr;
+
