Copying //UE4/Dev-Online to //UE4/Dev-Main (Source: //UE4/Dev-Online @ 4040378)

============================
  MAJOR FEATURES & CHANGES
============================

Change 3661955 by [REDACTED_PERSON]

	Change access pattern to FHttpThread to remove critical section
	#http #ue4

Change 3672463 by [REDACTED_PERSON]

	Revert Curl/WinINet HTTP implementation of GetHeader/GetAllHeaders to not allowing access until the request is complete
	Some optimizations to Curl/WinINet GetAllHeaders
	Add bIsAsyncProcessingFinished to FHttpResponseWinInet, set when the request is finished on the HTTP thread, and only set bIsReady when we have finished processing all received headers
	Use FThreadSafeBool instead of volatile int32 in WinINet HTTP

Change 3855724 by [REDACTED_PERSON]

	Fix include paths

Change 3949903 by [REDACTED_PERSON]

	#OnlineSubsystemSteam - Add IsValid check to SteamSessionInfo to fix a rare crash
	- Misc cleanup / casting fixes
	#github #4532

Change 3949914 by [REDACTED_PERSON]

	#ShooterGame - Wait for session creation to finish before cleaning up session on network errors

Change 3662317 by [REDACTED_PERSON]

	Add delegate for when we receive an http header
	Add FHttpRequestImpl to implement simple virtual functions that all platforms implement in the same way
	Add lock access to Curl/WinInet response headers access so game thread can safely read while processing
	#http #ue4

diff --git a/Engine/Plugins/Online/OnlineSubsystemSteam/Source/Private/OnlineSessionAsyncLobbySteam.cpp b/Engine/Plugins/Online/OnlineSubsystemSteam/Source/Private/OnlineSessionAsyncLobbySteam.cpp
--- a/Engine/Plugins/Online/OnlineSubsystemSteam/Source/Private/OnlineSessionAsyncLobbySteam.cpp
+++ b/Engine/Plugins/Online/OnlineSubsystemSteam/Source/Private/OnlineSessionAsyncLobbySteam.cpp
@@ -150,19 +150,19 @@ static void GetLobbyKeyValuePairsFromSessionSettings(const FOnlineSessionSetting
  * @param SessionInfo session info to get key, value pairs from
  * @param KeyValuePairs key value pair structure to add to
  */
-static void GetLobbyKeyValuePairsFromSessionInfo(const FOnlineSessionInfoSteam* SessionInfo, FSteamSessionKeyValuePairs& KeyValuePairs)
+static void GetLobbyKeyValuePairsFromSessionInfo(const FOnlineSessionInfoSteam& SessionInfo, FSteamSessionKeyValuePairs& KeyValuePairs)
 {
-	if (SessionInfo->HostAddr.IsValid())
+	if (SessionInfo.HostAddr.IsValid())
 	{
 		uint32 HostAddr;
-		SessionInfo->HostAddr->GetIp(HostAddr);
+		SessionInfo.HostAddr->GetIp(HostAddr);
 		KeyValuePairs.Add(STEAMKEY_HOSTIP, FString::FromInt(HostAddr));
-		KeyValuePairs.Add(STEAMKEY_HOSTPORT, FString::FromInt(SessionInfo->HostAddr->GetPort()));
+		KeyValuePairs.Add(STEAMKEY_HOSTPORT, FString::FromInt(SessionInfo.HostAddr->GetPort()));
 	}
 
-	if (SessionInfo->SteamP2PAddr.IsValid())
+	if (SessionInfo.SteamP2PAddr.IsValid())
 	{
-		TSharedPtr<FInternetAddrSteam> SteamAddr = StaticCastSharedPtr<FInternetAddrSteam>(SessionInfo->SteamP2PAddr);
+		TSharedPtr<FInternetAddrSteam> SteamAddr = StaticCastSharedPtr<FInternetAddrSteam>(SessionInfo.SteamP2PAddr);
 		KeyValuePairs.Add(STEAMKEY_P2PADDR, SteamAddr->ToString(false));
 		KeyValuePairs.Add(STEAMKEY_P2PPORT, FString::FromInt(SteamAddr->GetPort()));
 	}
@@ -176,16 +176,19 @@ static void GetLobbyKeyValuePairsFromSessionInfo(const FOnlineSessionInfoSteam*
  */
 static void GetLobbyKeyValuePairsFromSession(const FOnlineSession* Session, FSteamSessionKeyValuePairs& KeyValuePairs)
 {
-	FUniqueNetIdSteam* SteamId = (FUniqueNetIdSteam*)(Session->OwningUserId.Get());
-	KeyValuePairs.Add(STEAMKEY_OWNINGUSERID, SteamId->ToString());
+	TSharedPtr<const FUniqueNetIdSteam> SteamId = StaticCastSharedPtr<const FUniqueNetIdSteam>(Session->OwningUserId);
+	if (SteamId.IsValid())
+	{
+		KeyValuePairs.Add(STEAMKEY_OWNINGUSERID, SteamId->ToString());
+	}
 	KeyValuePairs.Add(STEAMKEY_OWNINGUSERNAME, Session->OwningUserName);
 	KeyValuePairs.Add(STEAMKEY_NUMOPENPRIVATECONNECTIONS, FString::FromInt(Session->NumOpenPrivateConnections));
 	KeyValuePairs.Add(STEAMKEY_NUMOPENPUBLICCONNECTIONS, FString::FromInt(Session->NumOpenPublicConnections));
 
 	if (Session->SessionInfo.IsValid())
 	{
-		FOnlineSessionInfoSteam* SessionInfo = (FOnlineSessionInfoSteam*)(Session->SessionInfo.Get());
-		GetLobbyKeyValuePairsFromSessionInfo(SessionInfo, KeyValuePairs);
+		TSharedPtr<FOnlineSessionInfoSteam> SessionInfo = StaticCastSharedPtr<FOnlineSessionInfoSteam>(Session->SessionInfo);
+		GetLobbyKeyValuePairsFromSessionInfo(*SessionInfo, KeyValuePairs);
 	}
 
 	GetLobbyKeyValuePairsFromSessionSettings(Session->SessionSettings, KeyValuePairs);
@@ -211,7 +214,7 @@ bool FillSessionFromLobbyData(FUniqueNetIdSteam& LobbyId, FOnlineSession& Sessio
 	Session.SessionSettings.Settings.Empty();
 
 	// Create the session info
-	TSharedPtr<FOnlineSessionInfoSteam> SessionInfo = MakeShareable(new FOnlineSessionInfoSteam(ESteamSession::LobbySession, FUniqueNetIdSteam(LobbyId)));
+	TSharedPtr<FOnlineSessionInfoSteam> SessionInfo = MakeShared<FOnlineSessionInfoSteam>(ESteamSession::LobbySession, FUniqueNetIdSteam(LobbyId));
 	TSharedRef<FInternetAddr> HostAddr = ISocketSubsystem::Get(PLATFORM_SOCKETSUBSYSTEM)->CreateInternetAddr();
 	TSharedRef<FInternetAddrSteam> SteamP2PAddr = MakeShareable(new FInternetAddrSteam);
 	
@@ -482,17 +485,17 @@ void FOnlineAsyncTaskSteamCreateLobby::Finalize()
 			FUniqueNetIdSteam LobbyId(CallbackResults.m_ulSteamIDLobby);
 
 			// Setup the host session info now that we have a lobby id
-			FOnlineSessionInfoSteam* NewSessionInfo = new FOnlineSessionInfoSteam(ESteamSession::LobbySession, LobbyId);
+			TSharedRef<FOnlineSessionInfoSteam> NewSessionInfo = MakeShared<FOnlineSessionInfoSteam>(ESteamSession::LobbySession, LobbyId);
 			NewSessionInfo->Init();
 			// Lobby sessions don't have a valid IP
 			NewSessionInfo->HostAddr = NULL;
 			// Copy the P2P addr
-			FInternetAddrSteam* SteamAddr = new FInternetAddrSteam(FUniqueNetIdSteam(SteamUser()->GetSteamID()));
+			TSharedRef<FInternetAddrSteam> SteamAddr = MakeShared<FInternetAddrSteam>(FUniqueNetIdSteam(SteamUser()->GetSteamID()));
 			SteamAddr->SetPort(Subsystem->GetGameServerGamePort());
-			NewSessionInfo->SteamP2PAddr = MakeShareable(SteamAddr);
-			
+			NewSessionInfo->SteamP2PAddr = SteamAddr;
+
 			// Set the info on the session
-			Session->SessionInfo = MakeShareable(NewSessionInfo);
+			Session->SessionInfo = NewSessionInfo;
 
 			// Set the game state as pending (not started)
 			Session->SessionState = EOnlineSessionState::Pending;
@@ -581,73 +584,75 @@ void FOnlineAsyncTaskSteamUpdateLobby::Tick()
 		FNamedOnlineSession* Session = SessionInt->GetNamedSession(SessionName);
 		if (Session)
 		{
-			FOnlineSessionInfoSteam* SessionInfo = (FOnlineSessionInfoSteam*)(Session->SessionInfo.Get());
-
-			bool bUsesPresence = Session->SessionSettings.bUsesPresence;
-			if (bUsesPresence != NewSessionSettings.bUsesPresence)
+			TSharedPtr<FOnlineSessionInfoSteam> SessionInfo = StaticCastSharedPtr<FOnlineSessionInfoSteam>(Session->SessionInfo);
+			if (SessionInfo.IsValid())
 			{
-				UE_LOG_ONLINE(Warning, TEXT("Can't change presence settings on existing session %s, ignoring."), *SessionName.ToString());
-			}
-
-			FSteamSessionKeyValuePairs OldKeyValuePairs;
-			GetLobbyKeyValuePairsFromSession(Session, OldKeyValuePairs);
+				bool bUsesPresence = Session->SessionSettings.bUsesPresence;
+				if (bUsesPresence != NewSessionSettings.bUsesPresence)
+				{
+					UE_LOG_ONLINE(Warning, TEXT("Can't change presence settings on existing session %s, ignoring."), *SessionName.ToString());
+				}
 
-			Session->SessionSettings = NewSessionSettings;
-			Session->SessionSettings.bUsesPresence = bUsesPresence;
+				FSteamSessionKeyValuePairs OldKeyValuePairs;
+				GetLobbyKeyValuePairsFromSession(Session, OldKeyValuePairs);
 
-			if (bUpdateOnlineData)
-			{
-				ISteamMatchmaking* SteamMatchmakingPtr = SteamMatchmaking();
-				check(SteamMatchmakingPtr);
+				Session->SessionSettings = NewSessionSettings;
+				Session->SessionSettings.bUsesPresence = bUsesPresence;
 
-				ELobbyType LobbyType = BuildLobbyType(&Session->SessionSettings);
-				if (SteamMatchmakingPtr->SetLobbyType(SessionInfo->SessionId, LobbyType))
+				if (bUpdateOnlineData)
 				{
-					int32 LobbyMemberCount = SteamMatchmakingPtr->GetNumLobbyMembers(SessionInfo->SessionId);
-					int32 MaxLobbyMembers = SteamMatchmakingPtr->GetLobbyMemberLimit(SessionInfo->SessionId);
-					bool bLobbyJoinable = Session->SessionSettings.bAllowJoinInProgress && (LobbyMemberCount < MaxLobbyMembers);
-					if (SteamMatchmakingPtr->SetLobbyJoinable(SessionInfo->SessionId, bLobbyJoinable))
+					ISteamMatchmaking* SteamMatchmakingPtr = SteamMatchmaking();
+					check(SteamMatchmakingPtr);
+
+					ELobbyType LobbyType = BuildLobbyType(&Session->SessionSettings);
+					if (SteamMatchmakingPtr->SetLobbyType(SessionInfo->SessionId, LobbyType))
 					{
-						int32 NumConnections = Session->SessionSettings.NumPrivateConnections + Session->SessionSettings.NumPublicConnections;
-						if (SteamMatchmakingPtr->SetLobbyMemberLimit(SessionInfo->SessionId, NumConnections))
+						int32 LobbyMemberCount = SteamMatchmakingPtr->GetNumLobbyMembers(SessionInfo->SessionId);
+						int32 MaxLobbyMembers = SteamMatchmakingPtr->GetLobbyMemberLimit(SessionInfo->SessionId);
+						bool bLobbyJoinable = Session->SessionSettings.bAllowJoinInProgress && (LobbyMemberCount < MaxLobbyMembers);
+						if (SteamMatchmakingPtr->SetLobbyJoinable(SessionInfo->SessionId, bLobbyJoinable))
 						{
-							bWasSuccessful = true;
+							int32 NumConnections = Session->SessionSettings.NumPrivateConnections + Session->SessionSettings.NumPublicConnections;
+							if (SteamMatchmakingPtr->SetLobbyMemberLimit(SessionInfo->SessionId, NumConnections))
+							{
+								bWasSuccessful = true;
 
-							FSteamSessionKeyValuePairs KeyValuePairs;
-							GetLobbyKeyValuePairsFromSession(Session, KeyValuePairs);
+								FSteamSessionKeyValuePairs KeyValuePairs;
+								GetLobbyKeyValuePairsFromSession(Session, KeyValuePairs);
 
-							// @TODO ONLINE Make sure to only remove/set data that has changed
-							// Unregister old session properties with Steam lobby
-							for (FSteamSessionKeyValuePairs::TConstIterator It(OldKeyValuePairs); It; ++It)
-							{
-								UE_LOG_ONLINE(Verbose, TEXT("Removing Lobby Data (%s, %s)"), *It.Key(), *It.Value());
-								if (!SteamMatchmakingPtr->SetLobbyData(SessionInfo->SessionId, TCHAR_TO_UTF8(*It.Key()), ""))
+								// @TODO ONLINE Make sure to only remove/set data that has changed
+								// Unregister old session properties with Steam lobby
+								for (FSteamSessionKeyValuePairs::TConstIterator It(OldKeyValuePairs); It; ++It)
 								{
-									bWasSuccessful = false;
-									break;
+									UE_LOG_ONLINE(Verbose, TEXT("Removing Lobby Data (%s, %s)"), *It.Key(), *It.Value());
+									if (!SteamMatchmakingPtr->SetLobbyData(SessionInfo->SessionId, TCHAR_TO_UTF8(*It.Key()), ""))
+									{
+										bWasSuccessful = false;
+										break;
+									}
 								}
-							}
 
-							if (bWasSuccessful)
-							{
-								// Register session properties with Steam lobby
-								for (FSteamSessionKeyValuePairs::TConstIterator It(KeyValuePairs); It; ++It)
+								if (bWasSuccessful)
 								{
-									UE_LOG_ONLINE(Verbose, TEXT("Updating Lobby Data (%s, %s)"), *It.Key(), *It.Value());
-									if (!SteamMatchmakingPtr->SetLobbyData(SessionInfo->SessionId, TCHAR_TO_UTF8(*It.Key()), TCHAR_TO_UTF8(*It.Value())))
+									// Register session properties with Steam lobby
+									for (FSteamSessionKeyValuePairs::TConstIterator It(KeyValuePairs); It; ++It)
 									{
-										bWasSuccessful = false;
-										break;
+										UE_LOG_ONLINE(Verbose, TEXT("Updating Lobby Data (%s, %s)"), *It.Key(), *It.Value());
+										if (!SteamMatchmakingPtr->SetLobbyData(SessionInfo->SessionId, TCHAR_TO_UTF8(*It.Key()), TCHAR_TO_UTF8(*It.Value())))
+										{
+											bWasSuccessful = false;
+											break;
+										}
 									}
 								}
 							}
 						}
 					}
 				}
-			}
-			else
-			{
-				bWasSuccessful = true;
+				else
+				{
+					bWasSuccessful = true;
+				}
 			}
 		}
 	}
diff --git a/Engine/Source/Runtime/Online/BuildPatchServices/Private/Tests/Mock/HttpRequest.mock.h b/Engine/Source/Runtime/Online/BuildPatchServices/Private/Tests/Mock/HttpRequest.mock.h
--- a/Engine/Source/Runtime/Online/BuildPatchServices/Private/Tests/Mock/HttpRequest.mock.h
+++ b/Engine/Source/Runtime/Online/BuildPatchServices/Private/Tests/Mock/HttpRequest.mock.h
@@ -111,6 +111,11 @@ namespace BuildPatchServices
 			return HttpRequestProgressDelegate;
 		}
 
+		virtual FHttpRequestHeaderReceivedDelegate& OnHeaderReceived() override
+		{
+			return HttpHeaderReceivedDelegate;
+		}
+
 		virtual void CancelRequest() override
 		{
 			++RxCancelRequest;
@@ -142,6 +147,8 @@ namespace BuildPatchServices
 	public:
 		FHttpRequestProgressDelegate HttpRequestProgressDelegate;
 		FHttpRequestCompleteDelegate HttpRequestCompleteDelegate;
+		FHttpRequestHeaderReceivedDelegate HttpHeaderReceivedDelegate;
+
 		TArray<FRxSetVerb> RxSetVerb;
 		TArray<FRxSetURL> RxSetURL;
 		int32 RxProcessRequest;
diff --git a/Engine/Source/Runtime/Online/HTTP/Private/Apple/AppleHTTP.cpp b/Engine/Source/Runtime/Online/HTTP/Private/Apple/AppleHTTP.cpp
--- a/Engine/Source/Runtime/Online/HTTP/Private/Apple/AppleHTTP.cpp
+++ b/Engine/Source/Runtime/Online/HTTP/Private/Apple/AppleHTTP.cpp
@@ -234,20 +234,6 @@ bool FAppleHttpRequest::ProcessRequest()
 	return bStarted;
 }
 
-
-FHttpRequestCompleteDelegate& FAppleHttpRequest::OnProcessRequestComplete()
-{
-	UE_LOG(LogHttp, Verbose, TEXT("FAppleHttpRequest::OnProcessRequestComplete()"));
-	return RequestCompleteDelegate;
-}
-
-FHttpRequestProgressDelegate& FAppleHttpRequest::OnRequestProgress() 
-{
-	UE_LOG(LogHttp, VeryVerbose, TEXT("FAppleHttpRequest::OnRequestProgress()"));
-	return RequestProgressDelegate;
-}
-
-
 bool FAppleHttpRequest::StartRequest()
 {
 	SCOPED_AUTORELEASE_POOL;
@@ -310,6 +296,8 @@ void FAppleHttpRequest::FinishedRequest()
 		UE_LOG(LogHttp, Verbose, TEXT("Request succeeded"));
 		CompletionStatus = EHttpRequestStatus::Succeeded;
 
+		// TODO: Try to broadcast OnHeaderReceived when we receive headers instead of here at the end
+		BroadcastResponseHeadersReceived();
 		OnProcessRequestComplete().ExecuteIfBound(SharedThis(this), Response, true);
 	}
 	else
diff --git a/Engine/Source/Runtime/Online/HTTP/Private/Apple/AppleHTTP.h b/Engine/Source/Runtime/Online/HTTP/Private/Apple/AppleHTTP.h
--- a/Engine/Source/Runtime/Online/HTTP/Private/Apple/AppleHTTP.h
+++ b/Engine/Source/Runtime/Online/HTTP/Private/Apple/AppleHTTP.h
@@ -4,7 +4,7 @@
 
 #include "CoreMinimal.h"
 #include "Containers/UnrealString.h"
-#include "Interfaces/IHttpRequest.h"
+#include "GenericPlatform/HttpRequestImpl.h"
 #include "Interfaces/IHttpResponse.h"
 #include "HttpManager.h"
 #include "PlatformHttp.h"
@@ -16,7 +16,7 @@ class IHttpRequest;
 /**
  * Apple implementation of an Http request
  */
-class FAppleHttpRequest : public IHttpRequest
+class FAppleHttpRequest : public FHttpRequestImpl
 {
 private:
 	// This is the NSMutableURLRequest, all our Apple functionality will deal with this.
@@ -50,8 +50,6 @@ public:
 	virtual void SetHeader(const FString& HeaderName, const FString& HeaderValue) override;
 	virtual void AppendToHeader(const FString& HeaderName, const FString& AdditionalHeaderValue) override;
 	virtual bool ProcessRequest() override;
-	virtual FHttpRequestCompleteDelegate& OnProcessRequestComplete() override;
-	virtual FHttpRequestProgressDelegate& OnRequestProgress() override;
 	virtual void CancelRequest() override;
 	virtual EHttpRequestStatus::Type GetStatus() const override;
 	virtual const FHttpResponsePtr GetResponse() const override;
@@ -59,7 +57,6 @@ public:
 	virtual float GetElapsedTime() const override;
 	//~ End IHttpRequest Interface
 
-
 	/**
 	 * Constructor
 	 */
@@ -99,12 +96,6 @@ private:
 	/** BYTE array payload to use with the request. Typically for a POST */
 	mutable TArray<uint8> RequestPayload;
 
-	/** Delegate that will get called once request completes or on any error */
-	FHttpRequestCompleteDelegate RequestCompleteDelegate;
-
-	/** Delegate that will get called once per tick with bytes downloaded so far */
-	FHttpRequestProgressDelegate RequestProgressDelegate;
-
 	/** Current status of request being processed */
 	EHttpRequestStatus::Type CompletionStatus;
 
diff --git a/Engine/Source/Runtime/Online/HTTP/Private/Curl/CurlHttp.cpp b/Engine/Source/Runtime/Online/HTTP/Private/Curl/CurlHttp.cpp
--- a/Engine/Source/Runtime/Online/HTTP/Private/Curl/CurlHttp.cpp
+++ b/Engine/Source/Runtime/Online/HTTP/Private/Curl/CurlHttp.cpp
@@ -7,6 +7,7 @@
 #include "Http.h"
 #include "Misc/EngineVersion.h"
 #include "Curl/CurlHttpManager.h"
+#include "Misc/ScopeLock.h"
 
 #if WITH_LIBCURL
 
@@ -162,12 +163,25 @@ FString FCurlHttpRequest::GetHeader(const FString& HeaderName) const
 	return Result;
 }
 
+FString FCurlHttpRequest::CombineHeaderKeyValue(const FString& HeaderKey, const FString& HeaderValue)
+{
+	FString Combined;
+	const TCHAR Separator[] = TEXT(": ");
+	constexpr const int32 SeparatorLength = ARRAY_COUNT(Separator) - 1;
+	Combined.Reserve(HeaderKey.Len() + SeparatorLength + HeaderValue.Len());
+	Combined.Append(HeaderKey);
+	Combined.AppendChars(Separator, SeparatorLength);
+	Combined.Append(HeaderValue);
+	return Combined;
+}
+
 TArray<FString> FCurlHttpRequest::GetAllHeaders() const
 {
 	TArray<FString> Result;
-	for (TMap<FString, FString>::TConstIterator It(Headers); It; ++It)
+	Result.Reserve(Headers.Num());
+	for (const TPair<FString, FString>& It : Headers)
 	{
-		Result.Add(It.Key() + TEXT(": ") + It.Value());
+		Result.Emplace(CombineHeaderKeyValue(It.Key, It.Value));
 	}
 	return Result;
 }
@@ -309,23 +323,15 @@ size_t FCurlHttpRequest::ReceiveResponseHeaderCallback(void* Ptr, size_t SizeInB
 			FString HeaderKey, HeaderValue;
 			if (Header.Split(TEXT(":"), &HeaderKey, &HeaderValue))
 			{
+				HeaderValue.TrimStartInline();
 				if (!HeaderKey.IsEmpty() && !HeaderValue.IsEmpty())
 				{
-					FString* PreviousValue = Response->Headers.Find(HeaderKey);
-					FString NewValue;
-					if (PreviousValue != nullptr && !PreviousValue->IsEmpty())
-					{
-						NewValue = (*PreviousValue) + TEXT(", ");
-					}
-					HeaderValue.TrimStartInline();
-					NewValue += HeaderValue;
-					Response->Headers.Add(HeaderKey, NewValue);
-
 					//Store the content length so OnRequestProgress() delegates have something to work with
 					if (HeaderKey == TEXT("Content-Length"))
 					{
 						Response->ContentLength = FCString::Atoi(*HeaderValue);
 					}
+					Response->NewlyReceivedHeaders.Enqueue(TPair<FString, FString>(MoveTemp(HeaderKey), MoveTemp(HeaderValue)));
 				}
 			}
 			return HeaderSize;
@@ -733,16 +739,6 @@ void FCurlHttpRequest::TickThreadedRequest(float DeltaSeconds)
 	TimeSinceLastResponse += DeltaSeconds;
 }
 
-FHttpRequestCompleteDelegate& FCurlHttpRequest::OnProcessRequestComplete()
-{
-	return RequestCompleteDelegate;
-}
-
-FHttpRequestProgressDelegate& FCurlHttpRequest::OnRequestProgress()
-{
-	return RequestProgressDelegate;
-}
-
 void FCurlHttpRequest::CancelRequest()
 {
 	bCanceled = true;
@@ -773,18 +769,20 @@ const FHttpResponsePtr FCurlHttpRequest::GetResponse() const
 void FCurlHttpRequest::Tick(float DeltaSeconds)
 {
 	CheckProgressDelegate();
+	BroadcastNewlyReceivedHeaders();
 }
 
 void FCurlHttpRequest::CheckProgressDelegate()
-	{
+{
 	QUICK_SCOPE_CYCLE_COUNTER(STAT_FCurlHttpRequest_CheckProgressDelegate);
-	int32 CurrentBytesRead = Response.IsValid() ? Response->TotalBytesRead.GetValue() : 0;
-	int32 CurrentBytesSent = BytesSent.GetValue();
+	const int32 CurrentBytesRead = Response.IsValid() ? Response->TotalBytesRead.GetValue() : 0;
+	const int32 CurrentBytesSent = BytesSent.GetValue();
 
 	const bool bProcessing = CompletionStatus == EHttpRequestStatus::Processing;
-	const bool bProgessChanged = (CurrentBytesSent != LastReportedBytesSent) ||
-		(Response.IsValid() && CurrentBytesRead != LastReportedBytesRead);
-	if (bProcessing && bProgessChanged)
+	const bool bBytesSentChanged = (CurrentBytesSent != LastReportedBytesSent);
+	const bool bBytesReceivedChanged = (Response.IsValid() && CurrentBytesRead != LastReportedBytesRead);
+	const bool bProgressChanged = bBytesSentChanged || bBytesReceivedChanged;
+	if (bProcessing && bProgressChanged)
 	{
 		LastReportedBytesSent = CurrentBytesSent;
 		if (Response.IsValid())
@@ -796,8 +794,38 @@ void FCurlHttpRequest::CheckProgressDelegate()
 	}
 }
 
+void FCurlHttpRequest::BroadcastNewlyReceivedHeaders()
+{
+	check(IsInGameThread());
+	if (Response.IsValid())
+	{
+		// Process the headers received on the HTTP thread and merge them into our master list and then broadcast the new headers
+		TPair<FString, FString> NewHeader;
+		while (Response->NewlyReceivedHeaders.Dequeue(NewHeader))
+		{
+			const FString& HeaderKey = NewHeader.Key;
+			const FString& HeaderValue = NewHeader.Value;
+
+			FString NewValue;
+			FString* PreviousValue = Response->Headers.Find(HeaderKey);
+			if (PreviousValue != nullptr && !PreviousValue->IsEmpty())
+			{
+				constexpr const int32 SeparatorLength = 2; // Length of ", "
+				NewValue = MoveTemp(*PreviousValue);
+				NewValue.Reserve(NewValue.Len() + SeparatorLength + HeaderValue.Len());
+				NewValue += TEXT(", ");
+			}
+			NewValue += HeaderValue;
+			Response->Headers.Add(HeaderKey, MoveTemp(NewValue));
+
+			OnHeaderReceived().ExecuteIfBound(SharedThis(this), NewHeader.Key, NewHeader.Value);
+		}
+	}
+}
+
 void FCurlHttpRequest::FinishedRequest()
 {
+	check(IsInGameThread());
 	QUICK_SCOPE_CYCLE_COUNTER(STAT_FCurlHttpRequest_FinishedRequest);
 	
 	CheckProgressDelegate();
@@ -810,13 +838,13 @@ void FCurlHttpRequest::FinishedRequest()
 
 			// get the information
 			long HttpCode = 0;
-			if (0 == curl_easy_getinfo(EasyHandle, CURLINFO_RESPONSE_CODE, &HttpCode))
+			if (CURLE_OK == curl_easy_getinfo(EasyHandle, CURLINFO_RESPONSE_CODE, &HttpCode))
 			{
 				Response->HttpCode = HttpCode;
 			}
 
 			double ContentLengthDownload = 0.0;
-			if (0 == curl_easy_getinfo(EasyHandle, CURLINFO_CONTENT_LENGTH_DOWNLOAD, &ContentLengthDownload))
+			if (CURLE_OK == curl_easy_getinfo(EasyHandle, CURLINFO_CONTENT_LENGTH_DOWNLOAD, &ContentLengthDownload))
 			{
 				Response->ContentLength = static_cast< int32 >(ContentLengthDownload);
 			}
@@ -826,6 +854,7 @@ void FCurlHttpRequest::FinishedRequest()
 	// if just finished, mark as stopped async processing
 	if (Response.IsValid())
 	{
+		BroadcastNewlyReceivedHeaders();
 		Response->bIsReady = true;
 	}
 
@@ -867,9 +896,10 @@ void FCurlHttpRequest::FinishedRequest()
 			}
 		}
 
-
 		// Mark last request attempt as completed successfully
 		CompletionStatus = EHttpRequestStatus::Succeeded;
+		// Broadcast any headers we haven't broadcast yet
+		BroadcastNewlyReceivedHeaders();
 		// Call delegate with valid request/response objects
 		OnProcessRequestComplete().ExecuteIfBound(SharedThis(this),Response,true);
 	}
@@ -881,7 +911,7 @@ void FCurlHttpRequest::FinishedRequest()
 		}
 		else
 		{
-		UE_LOG(LogHttp, Warning, TEXT("%p: request failed, libcurl error: %d (%s)"), this, (int32)CurlCompletionResult, ANSI_TO_TCHAR(curl_easy_strerror(CurlCompletionResult)));
+			UE_LOG(LogHttp, Warning, TEXT("%p: request failed, libcurl error: %d (%s)"), this, (int32)CurlCompletionResult, ANSI_TO_TCHAR(curl_easy_strerror(CurlCompletionResult)));
 		}
 
 		for (int32 i = 0; i < InfoMessageCache.Num(); ++i)
@@ -946,7 +976,7 @@ FString FCurlHttpResponse::GetURLParameter(const FString& ParameterName) const
 
 FString FCurlHttpResponse::GetHeader(const FString& HeaderName) const
 {
-	FString Result(TEXT(""));
+	FString Result;
 	if (!bIsReady)
 	{
 		UE_LOG(LogHttp, Warning, TEXT("Can't get cached header [%s]. Response still processing. %p"),
@@ -957,7 +987,7 @@ FString FCurlHttpResponse::GetHeader(const FString& HeaderName) const
 		const FString* Header = Headers.Find(HeaderName);
 		if (Header != NULL)
 		{
-			return *Header;
+			Result = *Header;
 		}
 	}
 	return Result;
@@ -972,9 +1002,10 @@ TArray<FString> FCurlHttpResponse::GetAllHeaders() const
 	}
 	else
 	{
-		for (TMap<FString, FString>::TConstIterator It(Headers); It; ++It)
+		Result.Reserve(Headers.Num());
+		for (const TPair<FString, FString>& It : Headers)
 		{
-			Result.Add(It.Key() + TEXT(": ") + It.Value());
+			Result.Emplace(FCurlHttpRequest::CombineHeaderKeyValue(It.Key, It.Value));
 		}
 	}
 	return Result;
diff --git a/Engine/Source/Runtime/Online/HTTP/Private/Curl/CurlHttp.h b/Engine/Source/Runtime/Online/HTTP/Private/Curl/CurlHttp.h
--- a/Engine/Source/Runtime/Online/HTTP/Private/Curl/CurlHttp.h
+++ b/Engine/Source/Runtime/Online/HTTP/Private/Curl/CurlHttp.h
@@ -6,6 +6,7 @@
 #include "HAL/ThreadSafeCounter.h"
 #include "Interfaces/IHttpResponse.h"
 #include "IHttpThreadedRequest.h"
+#include "Containers/Queue.h"
 
 class FCurlHttpResponse;
 
@@ -141,8 +142,6 @@ public:
 	virtual void SetHeader(const FString& HeaderName, const FString& HeaderValue) override;
 	virtual void AppendToHeader(const FString& HeaderName, const FString& AdditionalHeaderValue) override;
 	virtual bool ProcessRequest() override;
-	virtual FHttpRequestCompleteDelegate& OnProcessRequestComplete() override;
-	virtual FHttpRequestProgressDelegate& OnRequestProgress() override;
 	virtual void CancelRequest() override;
 	virtual EHttpRequestStatus::Type GetStatus() const override;
 	virtual const FHttpResponsePtr GetResponse() const override;
@@ -199,7 +198,6 @@ public:
 	 */
 	virtual ~FCurlHttpRequest();
 
-
 private:
 
 	/**
@@ -303,16 +301,17 @@ private:
 	 */
 	void FinishedRequest();
 
-	/**
-	 * Close session/request handles and unregister callbacks
-	 */
-	void CleanupRequest();
-
 	/**
 	 * Trigger the request progress delegate if progress has changed
 	 */
 	void CheckProgressDelegate();
 
+	/** Broadcast newly received headers */
+	void BroadcastNewlyReceivedHeaders();
+
+	/** Combine a header's key/value in the format "Key: Value" */
+	static FString CombineHeaderKeyValue(const FString& HeaderKey, const FString& HeaderValue);
+	
 private:
 
 	/** Pointer to an easy handle specific to this request */
@@ -335,10 +334,6 @@ private:
 	TSharedPtr<class FCurlHttpResponse,ESPMode::ThreadSafe> Response;
 	/** BYTE array payload to use with the request. Typically for a POST */
 	TArray<uint8> RequestPayload;
-	/** Delegate that will get called once request completes or on any error */
-	FHttpRequestCompleteDelegate RequestCompleteDelegate;
-	/** Delegate that will get called once per tick with total bytes uploaded and downloaded so far */
-	FHttpRequestProgressDelegate RequestProgressDelegate;
 	/** Current status of request being processed */
 	EHttpRequestStatus::Type CompletionStatus;
 	/** Mapping of header section to values. */
@@ -361,8 +356,6 @@ private:
 	TArray<FString> InfoMessageCache;
 };
 
-
-
 /**
  * Curl implementation of an HTTP response
  */
@@ -378,7 +371,6 @@ public:
 	// implementation friends
 	friend class FCurlHttpRequest;
 
-
 	//~ Begin IHttpBase Interface
 	virtual FString GetURL() const override;
 	virtual FString GetURLParameter(const FString& ParameterName) const override;
@@ -394,11 +386,6 @@ public:
 	virtual FString GetContentAsString() const override;
 	//~ End IHttpResponse Interface
 
-	/**
-	 * Check whether a response is ready or not.
-	 */
-	bool IsReady();
-
 	/**
 	 * Constructor
 	 *
@@ -417,8 +404,10 @@ private:
 	TArray<uint8> Payload;
 	/** Caches how many bytes of the response we've read so far */
 	FThreadSafeCounter TotalBytesRead;
-	/** Cached key/value header pairs. Parsed once request completes */
+	/** Cached key/value header pairs. Parsed once request completes. Only accessible on the game thread. */
 	TMap<FString, FString> Headers;
+	/** Newly received headers we need to inform listeners about */
+	TQueue<TPair<FString, FString>> NewlyReceivedHeaders;
 	/** Cached code from completed response */
 	int32 HttpCode;
 	/** Cached content length from completed response */
diff --git a/Engine/Source/Runtime/Online/HTTP/Private/GenericPlatform/GenericPlatformHttp.cpp b/Engine/Source/Runtime/Online/HTTP/Private/GenericPlatform/GenericPlatformHttp.cpp
--- a/Engine/Source/Runtime/Online/HTTP/Private/GenericPlatform/GenericPlatformHttp.cpp
+++ b/Engine/Source/Runtime/Online/HTTP/Private/GenericPlatform/GenericPlatformHttp.cpp
@@ -1,7 +1,7 @@
 // Copyright 1998-2018 Epic Games, Inc. All Rights Reserved.
 
 #include "GenericPlatform/GenericPlatformHttp.h"
-#include "Interfaces/IHttpRequest.h"
+#include "GenericPlatform/HttpRequestImpl.h"
 #include "Misc/Paths.h"
 #include "Misc/EngineVersion.h"
 #include "Misc/App.h"
@@ -9,7 +9,7 @@
 /**
  * A generic http request
  */
-class FGenericPlatformHttpRequest : public IHttpRequest
+class FGenericPlatformHttpRequest : public FHttpRequestImpl
 {
 public:
 	// IHttpBase
@@ -30,8 +30,6 @@ public:
 	virtual void SetHeader(const FString& HeaderName, const FString& HeaderValue) override {}
 	virtual void AppendToHeader(const FString& HeaderName, const FString& AdditionalHeaderValue) override {}
 	virtual bool ProcessRequest() override { return false; }
-	virtual FHttpRequestCompleteDelegate& OnProcessRequestComplete() override { static FHttpRequestCompleteDelegate RequestCompleteDelegate; return RequestCompleteDelegate; }
-	virtual FHttpRequestProgressDelegate& OnRequestProgress() override { static FHttpRequestProgressDelegate RequestProgressDelegate; return RequestProgressDelegate; }
 	virtual void CancelRequest() override {}
 	virtual EHttpRequestStatus::Type GetStatus() const override { return EHttpRequestStatus::NotStarted; }
 	virtual const FHttpResponsePtr GetResponse() const override { return nullptr; }
diff --git a/Engine/Source/Runtime/Online/HTTP/Private/GenericPlatform/HttpRequestImpl.cpp b/Engine/Source/Runtime/Online/HTTP/Private/GenericPlatform/HttpRequestImpl.cpp
new file mode 100644
--- /dev/null
+++ b/Engine/Source/Runtime/Online/HTTP/Private/GenericPlatform/HttpRequestImpl.cpp
@@ -0,0 +1,45 @@
+// Copyright 1998-2017 Epic Games, Inc. All Rights Reserved.
+
+#include "GenericPlatform/HttpRequestImpl.h"
+#include "Http.h"
+
+FHttpRequestCompleteDelegate& FHttpRequestImpl::OnProcessRequestComplete()
+{
+	UE_LOG(LogHttp, VeryVerbose, TEXT("FHttpRequestImpl::OnProcessRequestComplete()"));
+	return RequestCompleteDelegate;
+}
+
+FHttpRequestProgressDelegate& FHttpRequestImpl::OnRequestProgress() 
+{
+	UE_LOG(LogHttp, VeryVerbose, TEXT("FHttpRequestImpl::OnRequestProgress()"));
+	return RequestProgressDelegate;
+}
+
+FHttpRequestHeaderReceivedDelegate& FHttpRequestImpl::OnHeaderReceived()
+{
+	UE_LOG(LogHttp, VeryVerbose, TEXT("FHttpRequestImpl::OnHeaderReceived()"));
+	return HeaderReceivedDelegate;
+}
+
+void FHttpRequestImpl::BroadcastResponseHeadersReceived()
+{
+	if (OnHeaderReceived().IsBound())
+	{
+		const FHttpResponsePtr Response = GetResponse();
+		if (Response.IsValid())
+		{
+			const FHttpRequestPtr ThisPtr(SharedThis(this));
+			const TArray<FString> AllHeaders = Response->GetAllHeaders();
+			for (const FString& Header : AllHeaders)
+			{
+				FString HeaderName;
+				FString HeaderValue;
+				if (Header.Split(TEXT(":"), &HeaderName, &HeaderValue))
+				{
+					HeaderValue.TrimStartInline();
+					OnHeaderReceived().ExecuteIfBound(ThisPtr, HeaderName, HeaderValue);
+				}
+			}
+		}
+	}
+}
diff --git a/Engine/Source/Runtime/Online/HTTP/Private/HTML5/HTML5HTTP.cpp b/Engine/Source/Runtime/Online/HTTP/Private/HTML5/HTML5HTTP.cpp
--- a/Engine/Source/Runtime/Online/HTTP/Private/HTML5/HTML5HTTP.cpp
+++ b/Engine/Source/Runtime/Online/HTTP/Private/HTML5/HTML5HTTP.cpp
@@ -243,6 +243,7 @@ void FHTML5HttpRequest::ReceiveCallback(void *arg, void *buffer, uint32 size, vo
 					HeaderValue.TrimStartInline();
 					NewValue += HeaderValue;
 					Response->Headers.Add(HeaderKey, NewValue);
+					OnHeaderReceived().ExecuteIfBound(SharedThis(this), HeaderKey, HeaderValue);
 				}
 			}
 			if (!HeaderLeftovers.IsEmpty())
@@ -463,19 +464,6 @@ bool FHTML5HttpRequest::ProcessRequest()
 	return true;
 }
 
-
-FHttpRequestCompleteDelegate& FHTML5HttpRequest::OnProcessRequestComplete()
-{
-	UE_LOG(LogHttp, Verbose, TEXT("FHTML5HttpRequest::OnProcessRequestComplete()"));
-	return RequestCompleteDelegate;
-}
-
-FHttpRequestProgressDelegate& FHTML5HttpRequest::OnRequestProgress()
-{
-	UE_LOG(LogHttp, Verbose, TEXT("FHTML5HttpRequest::OnRequestProgress()"));
-	return RequestProgressDelegate;
-}
-
 void FHTML5HttpRequest::FinishedRequest()
 {
 	UE_LOG(LogHttp, Verbose, TEXT("FHTML5HttpRequest::FinishedRequest()"));
diff --git a/Engine/Source/Runtime/Online/HTTP/Private/HTML5/HTML5HTTP.h b/Engine/Source/Runtime/Online/HTTP/Private/HTML5/HTML5HTTP.h
--- a/Engine/Source/Runtime/Online/HTTP/Private/HTML5/HTML5HTTP.h
+++ b/Engine/Source/Runtime/Online/HTTP/Private/HTML5/HTML5HTTP.h
@@ -5,13 +5,13 @@
 #include "CoreMinimal.h"
 #include "Containers/UnrealString.h"
 #include "Containers/Map.h"
-#include "Interfaces/IHttpRequest.h"
+#include "GenericPlatform/HttpRequestImpl.h"
 #include "Interfaces/IHttpResponse.h"
 
 /**
 * HTML5 implementation of an Http request
 */
-class FHTML5HttpRequest : public IHttpRequest
+class FHTML5HttpRequest : public FHttpRequestImpl
 {
 public:
 	// implementation friends
@@ -36,8 +36,6 @@ public:
 	virtual void SetHeader(const FString& HeaderName, const FString& HeaderValue) override;
 	virtual void AppendToHeader(const FString& HeaderName, const FString& AdditionalHeaderValue) override;
 	virtual bool ProcessRequest() override;
-	virtual FHttpRequestCompleteDelegate& OnProcessRequestComplete() override;
-	virtual FHttpRequestProgressDelegate& OnRequestProgress() override;
 	virtual void CancelRequest() override;
 	virtual EHttpRequestStatus::Type GetStatus() const override;
 	virtual const FHttpResponsePtr GetResponse() const override;
@@ -98,12 +96,6 @@ private:
 	/** BYTE array payload to use with the request. Typically for a POST */
 	TArray<uint8> RequestPayload;
 
-	/** Delegate that will get called once request completes or on any error */
-	FHttpRequestCompleteDelegate RequestCompleteDelegate;
-
-	/** Delegate that will get called once per tick with bytes downloaded so far */
-	FHttpRequestProgressDelegate RequestProgressDelegate;
-
 	/** Current status of request being processed */
 	EHttpRequestStatus::Type CompletionStatus;
 
diff --git a/Engine/Source/Runtime/Online/HTTP/Private/HttpRetrySystem.cpp b/Engine/Source/Runtime/Online/HTTP/Private/HttpRetrySystem.cpp
--- a/Engine/Source/Runtime/Online/HTTP/Private/HttpRetrySystem.cpp
+++ b/Engine/Source/Runtime/Online/HTTP/Private/HttpRetrySystem.cpp
@@ -418,6 +418,10 @@ bool FHttpRetrySystem::FManager::Update(uint32* FileCount, uint32* FailingCount,
 
 		if (bWasCompleted)
 		{
+			if (bWasSuccessful)
+			{
+				HttpRetryRequest->BroadcastResponseHeadersReceived();
+			}
 			HttpRetryRequest->OnProcessRequestComplete().ExecuteIfBound(HttpRetryRequest, HttpRetryRequest->GetResponse(), bWasSuccessful);
 		}
 
diff --git a/Engine/Source/Runtime/Online/HTTP/Private/HttpThread.cpp b/Engine/Source/Runtime/Online/HTTP/Private/HttpThread.cpp
--- a/Engine/Source/Runtime/Online/HTTP/Private/HttpThread.cpp
+++ b/Engine/Source/Runtime/Online/HTTP/Private/HttpThread.cpp
@@ -5,7 +5,6 @@
 #include "HAL/PlatformTime.h"
 #include "HAL/PlatformProcess.h"
 #include "HAL/RunnableThread.h"
-#include "Misc/ScopeLock.h"
 #include "HttpModule.h"
 #include "Http.h"
 
@@ -44,21 +43,22 @@ void FHttpThread::StopThread()
 
 void FHttpThread::AddRequest(IHttpThreadedRequest* Request)
 {
-	FScopeLock ScopeLock(&RequestArraysLock);
-	PendingThreadedRequests.Add(Request);
+	PendingThreadedRequests.Enqueue(Request);
 }
 
 void FHttpThread::CancelRequest(IHttpThreadedRequest* Request)
 {
-	FScopeLock ScopeLock(&RequestArraysLock);
-	CancelledThreadedRequests.Add(Request);
+	CancelledThreadedRequests.Enqueue(Request);
 }
 
 void FHttpThread::GetCompletedRequests(TArray<IHttpThreadedRequest*>& OutCompletedRequests)
 {
-	FScopeLock ScopeLock(&RequestArraysLock);
-	OutCompletedRequests = CompletedThreadedRequests;
-	CompletedThreadedRequests.Reset();
+	check(IsInGameThread());
+	IHttpThreadedRequest* Request = nullptr;
+	while (CompletedThreadedRequests.Dequeue(Request))
+	{
+		OutCompletedRequests.Add(Request);
+	}
 }
 
 bool FHttpThread::Init()
@@ -69,6 +69,7 @@ bool FHttpThread::Init()
 
 uint32 FHttpThread::Run()
 {
+	// Arrays declared outside of loop to re-use memory
 	TArray<IHttpThreadedRequest*> RequestsToCancel;
 	TArray<IHttpThreadedRequest*> RequestsToStart;
 	TArray<IHttpThreadedRequest*> RequestsToComplete;
@@ -134,13 +135,19 @@ void FHttpThread::Process(TArray<IHttpThreadedRequest*>& RequestsToCancel, TArra
 {
 	// cache all cancelled and pending requests
 	{
-		FScopeLock ScopeLock(&RequestArraysLock);
+		IHttpThreadedRequest* Request = nullptr;
 
-		RequestsToCancel = CancelledThreadedRequests;
-		CancelledThreadedRequests.Reset();
+		RequestsToCancel.Reset();
+		while (CancelledThreadedRequests.Dequeue(Request))
+		{
+			RequestsToCancel.Add(Request);
+		}
 
-		RequestsToStart = PendingThreadedRequests;
-		PendingThreadedRequests.Reset();
+		RequestsToStart.Reset();
+		while (PendingThreadedRequests.Dequeue(Request))
+		{
+			RequestsToStart.Add(Request);
+		}
 	}
 
 	// Cancel any pending cancel requests
@@ -200,11 +207,7 @@ void FHttpThread::Process(TArray<IHttpThreadedRequest*>& RequestsToCancel, TArra
 		for (IHttpThreadedRequest* Request : RequestsToComplete)
 		{
 			CompleteThreadedRequest(Request);
-		}
-
-		{
-			FScopeLock ScopeLock(&RequestArraysLock);
-			CompletedThreadedRequests.Append(RequestsToComplete);
+			CompletedThreadedRequests.Enqueue(Request);
 		}
 		RequestsToComplete.Reset();
 	}
diff --git a/Engine/Source/Runtime/Online/HTTP/Private/HttpThread.h b/Engine/Source/Runtime/Online/HTTP/Private/HttpThread.h
--- a/Engine/Source/Runtime/Online/HTTP/Private/HttpThread.h
+++ b/Engine/Source/Runtime/Online/HTTP/Private/HttpThread.h
@@ -7,6 +7,7 @@
 #include "HAL/Runnable.h"
 #include "HttpPackage.h"
 #include "Misc/SingleThreadRunnable.h"
+#include "Containers/Queue.h"
 
 class IHttpThreadedRequest;
 
@@ -108,32 +109,29 @@ protected:
 	double LastTime;
 
 protected:
-	/** Critical section to lock access to PendingThreadedRequests, CancelledThreadedRequests, and CompletedThreadedRequests */
-	FCriticalSection RequestArraysLock;
-
 	/** 
 	 * Threaded requests that are waiting to be processed on the http thread.
-	 * Added to on non-HTTP thread, processed then cleared on HTTP thread.
+	 * Added to on (any) non-HTTP thread, processed then cleared on HTTP thread.
 	 */
-	TArray<IHttpThreadedRequest*> PendingThreadedRequests;
+	TQueue<IHttpThreadedRequest*, EQueueMode::Mpsc> PendingThreadedRequests;
 
 	/**
 	 * Threaded requests that are waiting to be cancelled on the http thread.
-	 * Added to on non-HTTP thread, processed then cleared on HTTP thread.
+	 * Added to on (any) non-HTTP thread, processed then cleared on HTTP thread.
 	 */
-	TArray<IHttpThreadedRequest*> CancelledThreadedRequests;
+	TQueue<IHttpThreadedRequest*, EQueueMode::Mpsc> CancelledThreadedRequests;
 
 	/**
-	 * Currently running threaded requests (not in any of the other arrays).
+	 * Currently running threaded requests (not in any of the other lists, except potentially CancelledThreadedRequests).
 	 * Only accessed on the HTTP thread.
 	 */
 	TArray<IHttpThreadedRequest*> RunningThreadedRequests;
 
 	/**
 	 * Threaded requests that have completed and are waiting for the game thread to process.
-	 * Added to on HTTP thread, processed then cleared on game thread.
+	 * Added to on HTTP thread, processed then cleared on game thread (Single producer, single consumer)
 	 */
-	TArray<IHttpThreadedRequest*> CompletedThreadedRequests;
+	TQueue<IHttpThreadedRequest*, EQueueMode::Spsc> CompletedThreadedRequests;
 
 	/** Pointer to Runnable Thread */
 	FRunnableThread* Thread;
diff --git a/Engine/Source/Runtime/Online/HTTP/Private/IHttpThreadedRequest.h b/Engine/Source/Runtime/Online/HTTP/Private/IHttpThreadedRequest.h
--- a/Engine/Source/Runtime/Online/HTTP/Private/IHttpThreadedRequest.h
+++ b/Engine/Source/Runtime/Online/HTTP/Private/IHttpThreadedRequest.h
@@ -3,9 +3,9 @@
 #pragma once
 
 #include "CoreMinimal.h"
-#include "Interfaces/IHttpRequest.h"
+#include "GenericPlatform/HttpRequestImpl.h"
 
-class IHttpThreadedRequest : public IHttpRequest
+class IHttpThreadedRequest : public FHttpRequestImpl
 {
 public:
 	// Called on http thread
diff --git a/Engine/Source/Runtime/Online/HTTP/Private/NullHttp.cpp b/Engine/Source/Runtime/Online/HTTP/Private/NullHttp.cpp
--- a/Engine/Source/Runtime/Online/HTTP/Private/NullHttp.cpp
+++ b/Engine/Source/Runtime/Online/HTTP/Private/NullHttp.cpp
@@ -111,16 +111,6 @@ bool FNullHttpRequest::ProcessRequest()
 	return true;
 }
 
-FHttpRequestCompleteDelegate& FNullHttpRequest::OnProcessRequestComplete()
-{
-	return RequestCompleteDelegate;
-}
-
-FHttpRequestProgressDelegate& FNullHttpRequest::OnRequestProgress()
-{
-	return ReuestProgressDelegate;
-}
-
 void FNullHttpRequest::CancelRequest()
 {
 	FinishedRequest();
diff --git a/Engine/Source/Runtime/Online/HTTP/Private/NullHttp.h b/Engine/Source/Runtime/Online/HTTP/Private/NullHttp.h
--- a/Engine/Source/Runtime/Online/HTTP/Private/NullHttp.h
+++ b/Engine/Source/Runtime/Online/HTTP/Private/NullHttp.h
@@ -3,13 +3,13 @@
 #pragma once
 
 #include "CoreMinimal.h"
-#include "Interfaces/IHttpRequest.h"
+#include "GenericPlatform/HttpRequestImpl.h"
 #include "Interfaces/IHttpResponse.h"
 
 /**
  * Null (mock) implementation of an HTTP request
  */
-class FNullHttpRequest : public IHttpRequest
+class FNullHttpRequest : public FHttpRequestImpl
 {
 public:
 
@@ -30,8 +30,6 @@ public:
 	virtual void SetHeader(const FString& HeaderName, const FString& HeaderValue) override;
 	virtual void AppendToHeader(const FString& HeaderName, const FString& AdditionalHeaderValue) override;
 	virtual bool ProcessRequest() override;
-	virtual FHttpRequestCompleteDelegate& OnProcessRequestComplete() override;
-	virtual FHttpRequestProgressDelegate& OnRequestProgress() override;
 	virtual void CancelRequest() override;
 	virtual EHttpRequestStatus::Type GetStatus() const override;
 	virtual const FHttpResponsePtr GetResponse() const override;
@@ -50,8 +48,6 @@ private:
 	FString Url;
 	FString Verb;
 	TArray<uint8> Payload;
-	FHttpRequestCompleteDelegate RequestCompleteDelegate;
-	FHttpRequestProgressDelegate ReuestProgressDelegate;
 	EHttpRequestStatus::Type CompletionStatus;
 	TMap<FString, FString> Headers;
 	float ElapsedTime;
diff --git a/Engine/Source/Runtime/Online/HTTP/Private/Windows/HttpWinInet.cpp b/Engine/Source/Runtime/Online/HTTP/Private/Windows/HttpWinInet.cpp
--- a/Engine/Source/Runtime/Online/HTTP/Private/Windows/HttpWinInet.cpp
+++ b/Engine/Source/Runtime/Online/HTTP/Private/Windows/HttpWinInet.cpp
@@ -176,7 +176,7 @@ void CALLBACK InternetStatusCallbackWinInet(
 		if (Response != NULL)
 		{
 			// if we receive a response, we sent the request (and it's no longer safe to retry)
-			FPlatformAtomics::InterlockedExchange(&Response->bRequestSent, 1);
+			Response->bRequestSent = true;
 		}
 		break;
 	case INTERNET_STATUS_RESPONSE_RECEIVED:
@@ -197,17 +197,12 @@ void CALLBACK InternetStatusCallbackWinInet(
 				DEBUG_LOG_HTTP(bDebugLog, Log, TEXT("InternetStatusCallbackWinInet request=%p AsyncResult.dwError: %08X. "), 
 					dwContext, AsyncResult->dwError, *InternetTranslateError(AsyncResult->dwError));
 
-				if (Response != NULL)
-				{
-					// Done processing response due to error
-					FPlatformAtomics::InterlockedExchange(&Response->bIsReady, 1);
-				}
 				Response = NULL;
 			}
 		}
 		
 		if (Request != NULL && Request->GetStatus() == EHttpRequestStatus::Processing &&
-			Response != NULL && !Response->bIsReady)
+			Response != NULL && !Response->bIsAsyncProcessingFinished)
 		{
 			Response->ProcessResponse();
 		}
@@ -224,7 +219,7 @@ void CALLBACK InternetStatusCallbackWinInet(
 		if (Response != NULL)
 		{
 			// mark that we have started sending the request (at this point it's no longer safe to retry)
-			FPlatformAtomics::InterlockedExchange(&Response->bRequestSent, 1);
+			Response->bRequestSent = true;
 		}
 		break;
 	case INTERNET_STATUS_STATE_CHANGE:
@@ -387,12 +382,25 @@ FString FHttpRequestWinInet::GetHeader(const FString& HeaderName) const
 	return Header != NULL ? *Header : TEXT("");
 }
 
+FString FHttpRequestWinInet::CombineHeaderKeyValue(const FString& HeaderKey, const FString& HeaderValue)
+{
+	FString Combined;
+	const TCHAR Separator[] = TEXT(": ");
+	constexpr const int32 SeparatorLength = ARRAY_COUNT(Separator) - 1;
+	Combined.Reserve(HeaderKey.Len() + SeparatorLength + HeaderValue.Len());
+	Combined.Append(HeaderKey);
+	Combined.AppendChars(Separator, SeparatorLength);
+	Combined.Append(HeaderValue);
+	return Combined;
+}
+
 TArray<FString> FHttpRequestWinInet::GetAllHeaders() const
 {
 	TArray<FString> Result;
-	for (TMap<FString, FString>::TConstIterator It(RequestHeaders); It; ++It)
+	Result.Reserve(RequestHeaders.Num());
+	for (const TPair<FString, FString>& It : RequestHeaders)
 	{
-		Result.Add(It.Key() + TEXT(": ") + It.Value());
+		Result.Emplace(CombineHeaderKeyValue(It.Key, It.Value));
 	}
 	return Result;
 }
@@ -647,6 +655,7 @@ bool FHttpRequestWinInet::StartRequest()
 
 void FHttpRequestWinInet::FinishedRequest()
 {
+	check(IsInGameThread());
 	// Clean up session/request handles that may have been created
 	CleanupRequest();
 	TSharedRef<IHttpRequest> Request = SharedThis(this);
@@ -654,6 +663,14 @@ void FHttpRequestWinInet::FinishedRequest()
 	FHttpModule::Get().GetHttpManager().RemoveRequest(Request);
 
 	ElapsedTime = (float)(FPlatformTime::Seconds() - StartRequestTime);
+
+	// Broadcast out any headers we haven't already
+	BroadcastNewlyReceivedHeaders();
+	if (Response.IsValid())
+	{
+		Response->bIsReady = true;
+	}
+
 	if (Response.IsValid() &&
 		Response->bResponseSucceeded)
 	{
@@ -695,7 +712,7 @@ void FHttpRequestWinInet::FinishedRequest()
 		// Mark last request attempt as completed successfully
 		CompletionStatus = EHttpRequestStatus::Succeeded;
 		// Call delegate with valid request/response objects
-		OnProcessRequestComplete().ExecuteIfBound(Request,Response,true);
+		OnProcessRequestComplete().ExecuteIfBound(Request, Response, true);
 	}
 	else
 	{
@@ -705,9 +722,9 @@ void FHttpRequestWinInet::FinishedRequest()
 		// Mark last request attempt as completed but failed
 		CompletionStatus = (Response.IsValid() && Response->bRequestSent) ? EHttpRequestStatus::Failed : EHttpRequestStatus::Failed_ConnectionError;
 		// No response since connection failed
-		Response = NULL;
+		Response.Reset();
 		// Call delegate with failure
-		OnProcessRequestComplete().ExecuteIfBound(Request,NULL,false);
+		OnProcessRequestComplete().ExecuteIfBound(Request, nullptr, false);
 	}
 }
 
@@ -741,17 +758,7 @@ FString FHttpRequestWinInet::GenerateHeaderBuffer(uint32 ContentLength)
 
 void FHttpRequestWinInet::ResetRequestTimeout()
 {
-	FPlatformAtomics::InterlockedExchange(&ElapsedTimeSinceLastServerResponse, 0);
-}
-
-FHttpRequestCompleteDelegate& FHttpRequestWinInet::OnProcessRequestComplete()
-{
-	return RequestCompleteDelegate;
-}
-
-FHttpRequestProgressDelegate& FHttpRequestWinInet::OnRequestProgress()
-{
-	return RequestProgressDelegate;
+	ElapsedTimeSinceLastServerResponse.Set(0);
 }
 
 void FHttpRequestWinInet::CancelRequest()
@@ -774,21 +781,52 @@ const FHttpResponsePtr FHttpRequestWinInet::GetResponse() const
 	return Response;
 }
 
+void FHttpRequestWinInet::BroadcastNewlyReceivedHeaders()
+{
+	check(IsInGameThread());
+	if (Response.IsValid())
+	{
+		// Process the headers received on the HTTP thread and merge them into our master list and then broadcast the new headers
+		TPair<FString, FString> NewHeader;
+		while (Response->NewlyReceivedHeaders.Dequeue(NewHeader))
+		{
+			const FString& HeaderKey = NewHeader.Key;
+			const FString& HeaderValue = NewHeader.Value;
+
+			FString NewValue;
+			FString* PreviousValue = Response->ResponseHeaders.Find(HeaderKey);
+			if (PreviousValue != nullptr && !PreviousValue->IsEmpty())
+			{
+				constexpr const int32 SeparatorLength = 2; // Length of ", "
+				NewValue = MoveTemp(*PreviousValue);
+				NewValue.Reserve(NewValue.Len() + SeparatorLength + HeaderValue.Len());
+				NewValue += TEXT(", ");
+			}
+			NewValue += HeaderValue;
+			Response->ResponseHeaders.Add(HeaderKey, MoveTemp(NewValue));
+
+			OnHeaderReceived().ExecuteIfBound(SharedThis(this), NewHeader.Key, NewHeader.Value);
+		}
+	}
+}
+
 void FHttpRequestWinInet::Tick(float DeltaSeconds)
 {
 	// keep track of elapsed milliseconds
 	const int32 ElapsedMillisecondsThisFrame = DeltaSeconds * 1000;
-	FPlatformAtomics::InterlockedAdd(&ElapsedTimeSinceLastServerResponse, ElapsedMillisecondsThisFrame);
+	ElapsedTimeSinceLastServerResponse.Add(ElapsedMillisecondsThisFrame);
 
 	// Update response progress
 	if(Response.IsValid())
 	{
-		int32 ResponseBytes = Response->ProgressBytesRead.GetValue();
+		const int32 ResponseBytes = Response->ProgressBytesRead.GetValue();
 		if(ResponseBytes > ProgressBytesSent)
 		{
 			ProgressBytesSent = ResponseBytes;
 			OnRequestProgress().ExecuteIfBound(SharedThis(this), RequestPayload.Num(), ResponseBytes);
 		}
+
+		BroadcastNewlyReceivedHeaders();
 	}
 
 	// Convert to seconds for comparison to the timeout value
@@ -807,7 +845,7 @@ void FHttpRequestWinInet::Tick(float DeltaSeconds)
 	}
 	
 	// Timeout if no server response for HttpTimeout
-	const float SecondsSinceLastResponse = ElapsedTimeSinceLastServerResponse / 1000.f;
+	const float SecondsSinceLastResponse = ElapsedTimeSinceLastServerResponse.GetValue() / 1000.f;
 	if (HttpTimeout > 0 && 
 		SecondsSinceLastResponse >= HttpTimeout)
 	{
@@ -820,7 +858,7 @@ void FHttpRequestWinInet::Tick(float DeltaSeconds)
 	// No longer waiting for a response and done processing it
 	else if (CompletionStatus == EHttpRequestStatus::Processing &&
 		Response.IsValid() &&
-		Response->bIsReady &&
+		Response->bIsAsyncProcessingFinished &&
 		TotalElapsed >= FHttpModule::Get().GetHttpDelayTime())
 	{
 		FinishedRequest();
@@ -841,9 +879,10 @@ FHttpResponseWinInet::FHttpResponseWinInet(FHttpRequestWinInet& InRequest)
 ,	TotalBytesRead(0)
 ,	ResponseCode(EHttpResponseCodes::Unknown)
 ,	ContentLength(0)
-,	bIsReady(0)
-,	bResponseSucceeded(0)
-,	bRequestSent(0)
+,	bIsAsyncProcessingFinished(false)
+,	bIsReady(false)
+,	bResponseSucceeded(false)
+,	bRequestSent(false)
 ,	MaxReadBufferSize(FHttpModule::Get().GetMaxReadBufferSize())
 {
 
@@ -867,7 +906,7 @@ FString FHttpResponseWinInet::GetURLParameter(const FString& ParameterName) cons
 
 FString FHttpResponseWinInet::GetHeader(const FString& HeaderName) const
 {
-	FString Result(TEXT(""));
+	FString Result;
 	if (!bIsReady)
 	{
 		UE_LOG(LogHttp, Warning, TEXT("Can't get cached header [%s]. Response still processing. %p"),
@@ -878,7 +917,7 @@ FString FHttpResponseWinInet::GetHeader(const FString& HeaderName) const
 		const FString* Header = ResponseHeaders.Find(HeaderName);
 		if (Header != NULL)
 		{
-			return *Header;
+			Result = *Header;
 		}
 	}
 	return Result;
@@ -893,9 +932,10 @@ TArray<FString> FHttpResponseWinInet::GetAllHeaders() const
 	}
 	else
 	{
-		for (TMap<FString, FString>::TConstIterator It(ResponseHeaders); It; ++It)
+		Result.Reserve(ResponseHeaders.Num());
+		for (const TPair<FString, FString>& It : ResponseHeaders)
 		{
-			Result.Add(It.Key() + TEXT(": ") + It.Value());
+			Result.Emplace(FHttpRequestWinInet::CombineHeaderKeyValue(It.Key, It.Value));
 		}
 	}
 	return Result;
@@ -1036,12 +1076,13 @@ void FHttpResponseWinInet::ProcessResponse()
 	// Cache content length now that response is done
 	ContentLength = QueryContentLength();
 	// Mark as valid processed response
-	FPlatformAtomics::InterlockedExchange(&bResponseSucceeded, 1);
-	// Done processing
-	FPlatformAtomics::InterlockedExchange(&bIsReady, 1);
+	bResponseSucceeded = true;
 	// Update progress bytes
-	ProgressBytesRead.Set( TotalBytesRead );
+	ProgressBytesRead.Set(TotalBytesRead);
+	// Done processing
+	bIsAsyncProcessingFinished = true;
 }
+
 void FHttpResponseWinInet::ProcessResponseHeaders()
 {
 	::DWORD HeaderSize = 0;
@@ -1077,22 +1118,14 @@ void FHttpResponseWinInet::ProcessResponseHeaders()
 				DelimiterPtr = EndPtr;
 			}
 			FString HeaderLine(DelimiterPtr-HeaderPtr, HeaderPtr);
-			FString HeaderKey,HeaderValue;
+			FString HeaderKey, HeaderValue;
 			if (HeaderLine.Split(TEXT(":"), &HeaderKey, &HeaderValue, ESearchCase::CaseSensitive))
 			{
 				HeaderKey.TrimStartAndEndInline();
 				HeaderValue.TrimStartAndEndInline();
 				if (!HeaderKey.IsEmpty() && !HeaderValue.IsEmpty())
 				{
-					FString* PreviousValue = ResponseHeaders.Find(HeaderKey);
-					FString NewValue;
-					if (PreviousValue != nullptr && !PreviousValue->IsEmpty())
-					{
-						NewValue = (*PreviousValue) + TEXT(", ");
-					}
-					HeaderValue.TrimStartInline();
-					NewValue += HeaderValue;
-					ResponseHeaders.Add(HeaderKey, NewValue);
+					NewlyReceivedHeaders.Enqueue(TPair<FString, FString>(MoveTemp(HeaderKey), MoveTemp(HeaderValue)));
 				}
 			}
 			HeaderPtr = DelimiterPtr + 2;
diff --git a/Engine/Source/Runtime/Online/HTTP/Private/Windows/HttpWinInet.h b/Engine/Source/Runtime/Online/HTTP/Private/Windows/HttpWinInet.h
--- a/Engine/Source/Runtime/Online/HTTP/Private/Windows/HttpWinInet.h
+++ b/Engine/Source/Runtime/Online/HTTP/Private/Windows/HttpWinInet.h
@@ -2,9 +2,12 @@
 
 #pragma once
 
+#include "GenericPlatform/HttpRequestImpl.h"
 #include "Interfaces/IHttpResponse.h"
 #include "HttpModule.h"
 #include "HAL/ThreadSafeCounter.h"
+#include "HAL/ThreadSafeBool.h"
+#include "Containers/Queue.h"
 #include "Windows/WindowsHWrapper.h"
 #include "Windows/AllowWindowsPlatformTypes.h"
 	#include <wininet.h>
@@ -67,9 +70,8 @@ private:
 /**
  * WinInet implementation of an Http request
  */
-class FHttpRequestWinInet : public IHttpRequest
+class FHttpRequestWinInet : public FHttpRequestImpl
 {
-
 public:
 
 	// IHttpBase
@@ -92,8 +94,6 @@ public:
 	virtual void SetHeader(const FString& HeaderName, const FString& HeaderValue) override;
 	virtual void AppendToHeader(const FString& HeaderName, const FString& AdditionalHeaderValue) override;
 	virtual bool ProcessRequest() override;
-	virtual FHttpRequestCompleteDelegate& OnProcessRequestComplete() override;
-	virtual FHttpRequestProgressDelegate& OnRequestProgress() override;
 	virtual void CancelRequest() override;
 	virtual EHttpRequestStatus::Type GetStatus() const override;
 	virtual const FHttpResponsePtr GetResponse() const override;
@@ -145,6 +145,12 @@ private:
 	 */
 	void ResetRequestTimeout();
 
+	/** Broadcast newly received headers */
+	void BroadcastNewlyReceivedHeaders();
+
+	/** Combine a header's key/value in the format "Key: Value" */
+	static FString CombineHeaderKeyValue(const FString& HeaderKey, const FString& HeaderValue);
+	
 	// implementation friends
 	friend class FHttpResponseWinInet;
 	friend void CALLBACK InternetStatusCallbackWinInet(HINTERNET, DWORD_PTR,::DWORD,LPVOID,::DWORD);
@@ -161,10 +167,6 @@ private:
 	TArray<uint8> RequestPayload;
 	/** Holds response data that comes back from a successful request.  NULL if request can't connect */
 	TSharedPtr<class FHttpResponseWinInet,ESPMode::ThreadSafe> Response;
-	/** Delegate that will get called once request completes or on any error */
-	FHttpRequestCompleteDelegate RequestCompleteDelegate;
-	/** Delegate that will get called once per tick with bytes downloaded so far */
-	FHttpRequestProgressDelegate RequestProgressDelegate;
 	/** Current status of request being processed */
 	EHttpRequestStatus::Type CompletionStatus;
 	/** WinInet handle to a session connection opened via InternetConnect */
@@ -172,7 +174,7 @@ private:
 	/** WinInet handle to a request opened via HttpOpenRequest */
 	HINTERNET RequestHandle;
 	/** Total elapsed time in seconds since the last server response, in milliseconds */
-	int32 volatile ElapsedTimeSinceLastServerResponse;
+	FThreadSafeCounter ElapsedTimeSinceLastServerResponse;
 	/** Number of bytes sent to progress update */
 	int32 ProgressBytesSent;
 	/** Used to calculate total elapsed time for the request */
@@ -267,8 +269,10 @@ private:
 	int32 AsyncBytesRead;
 	/** Caches how many bytes of the response we've read so far */
 	int32 TotalBytesRead;
-	/** Cached key/value header pairs. Parsed from HttpQueryInfo results once request completes */
+	/** Cached key/value header pairs. Parsed from HttpQueryInfo results once request completes. Only accessible on the game thread. */
 	TMap<FString, FString> ResponseHeaders;
+	/** Newly received headers we need to inform listeners about */
+	TQueue<TPair<FString, FString>> NewlyReceivedHeaders;
 	/** Cached code from completed response */
 	int32 ResponseCode;
 	/** Cached content length from completed response */
@@ -276,11 +280,13 @@ private:
 	/** BYTE array to fill in as the response is read via InternetReadFile */
 	TArray<uint8> ResponsePayload;
 	/** True when the response has finished async processing */
-	int32 volatile bIsReady;
+	FThreadSafeBool bIsAsyncProcessingFinished;
+	/** True when we have called FinishedRequest() */
+	bool bIsReady;
 	/** True if the response was successfully received/processed */
-	int32 volatile bResponseSucceeded;
+	FThreadSafeBool bResponseSucceeded;
 	/** Whether or not the request was actually sent */
-	int32 volatile bRequestSent;
+	FThreadSafeBool bRequestSent;
 	/** Threadsafe counter for exchanging bytes read so far with main thread tick */
 	FThreadSafeCounter ProgressBytesRead;
 	/** Max buffer size for individual http reads */
diff --git a/Engine/Source/Runtime/Online/HTTP/Public/GenericPlatform/HttpRequestImpl.h b/Engine/Source/Runtime/Online/HTTP/Public/GenericPlatform/HttpRequestImpl.h
new file mode 100644
--- /dev/null
+++ b/Engine/Source/Runtime/Online/HTTP/Public/GenericPlatform/HttpRequestImpl.h
@@ -0,0 +1,34 @@
+// Copyright 1998-2017 Epic Games, Inc. All Rights Reserved.
+
+#pragma once
+
+#include "Interfaces/IHttpRequest.h"
+
+/**
+ * Contains implementation of some common functions that don't vary between implementation
+ */
+class HTTP_API FHttpRequestImpl : public IHttpRequest
+{
+public:
+	// IHttpRequest
+	virtual FHttpRequestCompleteDelegate& OnProcessRequestComplete() override;
+	virtual FHttpRequestProgressDelegate& OnRequestProgress() override;
+	virtual FHttpRequestHeaderReceivedDelegate& OnHeaderReceived() override;
+
+protected:
+	/** 
+	 * Broadcast all of our response's headers as having been received
+	 * Used when we don't know when we receive headers in our HTTP implementation
+	 */
+	void BroadcastResponseHeadersReceived();
+
+protected:
+	/** Delegate that will get called once request completes or on any error */
+	FHttpRequestCompleteDelegate RequestCompleteDelegate;
+
+	/** Delegate that will get called once per tick with bytes downloaded so far */
+	FHttpRequestProgressDelegate RequestProgressDelegate;
+
+	/** Delegate that will get called for each new header received */
+	FHttpRequestHeaderReceivedDelegate HeaderReceivedDelegate;
+};
diff --git a/Engine/Source/Runtime/Online/HTTP/Public/HttpRequestAdapter.h b/Engine/Source/Runtime/Online/HTTP/Public/HttpRequestAdapter.h
--- a/Engine/Source/Runtime/Online/HTTP/Public/HttpRequestAdapter.h
+++ b/Engine/Source/Runtime/Online/HTTP/Public/HttpRequestAdapter.h
@@ -3,14 +3,14 @@
 #pragma once
 
 #include "CoreMinimal.h"
-#include "Interfaces/IHttpRequest.h"
+#include "GenericPlatform/HttpRequestImpl.h"
 
 /** 
   * Adapter class for IHttpRequest abstract interface
   * does not fully expose the wrapped interface in the base. This allows client defined marshalling of the requests when end point permissions are at issue.
   */
 
-class FHttpRequestAdapterBase : public IHttpRequest
+class FHttpRequestAdapterBase : public FHttpRequestImpl
 {
 public:
     FHttpRequestAdapterBase(const TSharedRef<IHttpRequest>& InHttpRequest) 
diff --git a/Engine/Source/Runtime/Online/HTTP/Public/HttpRetrySystem.h b/Engine/Source/Runtime/Online/HTTP/Public/HttpRetrySystem.h
--- a/Engine/Source/Runtime/Online/HTTP/Public/HttpRetrySystem.h
+++ b/Engine/Source/Runtime/Online/HTTP/Public/HttpRetrySystem.h
@@ -85,8 +85,6 @@ namespace FHttpRetrySystem
 		// IHttpRequest interface
 		HTTP_API virtual bool ProcessRequest() override;
 		HTTP_API virtual void CancelRequest() override;
-		virtual FHttpRequestCompleteDelegate& OnProcessRequestComplete() override { return OnProcessRequestCompleteDelegate; }
-		virtual FHttpRequestProgressDelegate& OnRequestProgress() override { return OnProcessRequestProgressDelegate; }
 		virtual FHttpRequestWillRetryDelegate& OnRequestWillRetry() { return OnRequestWillRetryDelegate; }
 		
 		// FRequest
@@ -113,8 +111,6 @@ namespace FHttpRetrySystem
 		FRetryResponseCodes					 RetryResponseCodes;
         FRetryVerbs                          RetryVerbs;
 
-		FHttpRequestCompleteDelegate OnProcessRequestCompleteDelegate;
-		FHttpRequestProgressDelegate OnProcessRequestProgressDelegate;
 		FHttpRequestWillRetryDelegate OnRequestWillRetryDelegate;
 
 		FManager& RetryManager;
diff --git a/Engine/Source/Runtime/Online/HTTP/Public/Interfaces/IHttpRequest.h b/Engine/Source/Runtime/Online/HTTP/Public/Interfaces/IHttpRequest.h
--- a/Engine/Source/Runtime/Online/HTTP/Public/Interfaces/IHttpRequest.h
+++ b/Engine/Source/Runtime/Online/HTTP/Public/Interfaces/IHttpRequest.h
@@ -63,20 +63,29 @@ typedef TSharedPtr<class IHttpResponse,ESPMode::ThreadSafe> FHttpResponsePtr;
 /**
  * Delegate called when an Http request completes
  *
- * @param first parameter - original Http request that started things
- * @param second parameter - response received from the server if a successful connection was established
- * @param third parameter - indicates whether or not the request was able to connect successfully
+ * @param Request original Http request that started things
+ * @param Response response received from the server if a successful connection was established
+ * @param bConnectedSuccessfully - indicates whether or not the request was able to connect successfully
  */
-DECLARE_DELEGATE_ThreeParams(FHttpRequestCompleteDelegate, FHttpRequestPtr, FHttpResponsePtr, bool);
+DECLARE_DELEGATE_ThreeParams(FHttpRequestCompleteDelegate, FHttpRequestPtr /*Request*/, FHttpResponsePtr /*Response*/, bool /*bConnectedSuccessfully*/);
+
+/**
+ * Delegate called when an Http request receives a header
+ *
+ * @param Request original Http request that started things
+ * @param HeaderName the name of the header
+ * @param NewHeaderValue the value of the header
+ */
+DECLARE_DELEGATE_ThreeParams(FHttpRequestHeaderReceivedDelegate, FHttpRequestPtr /*Request*/, const FString& /*HeaderName*/, const FString& /*NewHeaderValue*/);
 
 /**
  * Delegate called per tick to update an Http request upload or download size progress
  *
- * @param first parameter - original Http request that started things
- * @param second parameter - the number of bytes sent / uploaded in the request so far.
- * @param third parameter - the number of bytes received / downloaded in the response so far.
+ * @param Request original Http request that started things
+ * @param BytesSent the number of bytes sent / uploaded in the request so far.
+ * @param BytesReceived the number of bytes received / downloaded in the response so far.
  */
-DECLARE_DELEGATE_ThreeParams(FHttpRequestProgressDelegate, FHttpRequestPtr, int32, int32);
+DECLARE_DELEGATE_ThreeParams(FHttpRequestProgressDelegate, FHttpRequestPtr /*Request*/, int32 /*BytesSent*/, int32 /*BytesReceived*/);
 
 /**
  * Interface for Http requests (created using FHttpFactory)
@@ -172,6 +181,11 @@ public:
 	 */
 	virtual FHttpRequestProgressDelegate& OnRequestProgress() = 0;
 
+	/** 
+	 * Delegate called to signal the receipt of a header.  See FHttpRequestHeaderReceivedDelegate
+	 */
+	virtual FHttpRequestHeaderReceivedDelegate& OnHeaderReceived() = 0;
+
 	/**
 	 * Called to cancel a request that is still being processed
 	 */
