SetDedicatedServer before login to fix bad functional usage.

diff --git a/Engine/Plugins/Online/OnlineSubsystemSteam/Source/Private/OnlineSessionAsyncServerSteam.cpp b/Engine/Plugins/Online/OnlineSubsystemSteam/Source/Private/OnlineSessionAsyncServerSteam.cpp
--- a/Engine/Plugins/Online/OnlineSubsystemSteam/Source/Private/OnlineSessionAsyncServerSteam.cpp
+++ b/Engine/Plugins/Online/OnlineSubsystemSteam/Source/Private/OnlineSessionAsyncServerSteam.cpp
@@ -165,9 +165,6 @@ void UpdatePublishedSettings(UWorld* World, FNamedOnlineSession* Session)
 	// @TODO ONLINE Password protected or not
 	SteamGameServerPtr->SetPasswordProtected(false);
 
-	// Dedicated server or not
-	SteamGameServerPtr->SetDedicatedServer(Session->SessionSettings.bIsDedicated ? true : false);
-
 	// Map name
 	FString MapName;
 	if (TempSessionSettings.Get(SETTING_MAPNAME, MapName) && !MapName.IsEmpty())
@@ -330,34 +327,38 @@ FString FOnlineAsyncTaskSteamCreateServer::ToString() const
  */
 void FOnlineAsyncTaskSteamCreateServer::Tick() 
 {
+	FOnlineSessionSteamPtr SessionInt = StaticCastSharedPtr<FOnlineSessionSteam>(Subsystem->GetSessionInterface());
 	if (!bInit)
 	{
 		ISteamGameServer* SteamGameServerPtr = SteamGameServer();
-		check(SteamGameServerPtr);
+		FNamedOnlineSession* Session = (SessionInt.IsValid()) ? SessionInt->GetNamedSession(SessionName) : nullptr;
+		if (Session != nullptr && SteamGameServerPtr != nullptr)
+		{
+			bool bWantsDedicated = Session->SessionSettings.bIsDedicated;
+			UE_LOG_ONLINE(Verbose, TEXT("Initializing Steam game server. Is dedicated? %d"), bWantsDedicated);
 
-		UE_LOG_ONLINE(Verbose, TEXT("Initializing Steam game server"));
+			SteamGameServerPtr->SetModDir(STEAMGAMEDIR);
+			SteamGameServerPtr->SetProduct(STEAMPRODUCTNAME);
+			SteamGameServerPtr->SetGameDescription(STEAMGAMEDESC);
+			SteamGameServerPtr->SetDedicatedServer(bWantsDedicated);
+
+			if (!SteamGameServerPtr->BLoggedOn())
+			{
+				// Login the server with Steam
+				SteamGameServerPtr->LogOnAnonymous();
+			}
 
-		SteamGameServerPtr->SetModDir(STEAMGAMEDIR);
-		SteamGameServerPtr->SetProduct(STEAMPRODUCTNAME);
-		SteamGameServerPtr->SetGameDescription(STEAMGAMEDESC);
+			// Setup advertisement and force the initial update
+			SteamGameServerPtr->SetHeartbeatInterval(-1);
+			SteamGameServerPtr->EnableHeartbeats(true);
+			SteamGameServerPtr->ForceHeartbeat();
 
-		if (!SteamGameServerPtr->BLoggedOn())
-		{
-			// Login the server with Steam
-			SteamGameServerPtr->LogOnAnonymous();
+			bInit = true;
 		}
-
-		// Setup advertisement and force the initial update
-		SteamGameServerPtr->SetHeartbeatInterval(-1);
-		SteamGameServerPtr->EnableHeartbeats(true);
-		SteamGameServerPtr->ForceHeartbeat();
-		
-		bInit = true;
 	}
 
 	// Wait for the connection and policy response callbacks
-	FOnlineSessionSteamPtr SessionInt = StaticCastSharedPtr<FOnlineSessionSteam>(Subsystem->GetSessionInterface());
-	if (SessionInt->bSteamworksGameServerConnected && SessionInt->GameServerSteamId->IsValid() && SessionInt->bPolicyResponseReceived)
+	if (bInit && SessionInt->bSteamworksGameServerConnected && SessionInt->GameServerSteamId->IsValid() && SessionInt->bPolicyResponseReceived)
 	{
 		bIsComplete = true;
 		bWasSuccessful = true;
