Fixing the SteamAuth interface from not shutting off if the packet handlers do not exist.

diff --git a/Engine/Plugins/Online/OnlineSubsystemSteam/Source/Private/OnlineAuthHandlerSteam.cpp b/Engine/Plugins/Online/OnlineSubsystemSteam/Source/Private/OnlineAuthHandlerSteam.cpp
--- a/Engine/Plugins/Online/OnlineSubsystemSteam/Source/Private/OnlineAuthHandlerSteam.cpp
+++ b/Engine/Plugins/Online/OnlineSubsystemSteam/Source/Private/OnlineAuthHandlerSteam.cpp
@@ -95,7 +95,7 @@ FSteamAuthHandlerComponent::FSteamAuthHandlerComponent() :
 	if (OSS != nullptr)
 	{
 		AuthInterface = OSS->GetAuthInterface();
-		if (!AuthInterface.IsValid())
+		if (!AuthInterface.IsValid() || !AuthInterface->IsSessionAuthEnabled())
 		{
 			bIsEnabled = false;
 		}
diff --git a/Engine/Plugins/Online/OnlineSubsystemSteam/Source/Private/OnlineAuthInterfaceSteam.cpp b/Engine/Plugins/Online/OnlineSubsystemSteam/Source/Private/OnlineAuthInterfaceSteam.cpp
--- a/Engine/Plugins/Online/OnlineSubsystemSteam/Source/Private/OnlineAuthInterfaceSteam.cpp
+++ b/Engine/Plugins/Online/OnlineSubsystemSteam/Source/Private/OnlineAuthInterfaceSteam.cpp
@@ -13,6 +13,9 @@
 #include "Engine/World.h"
 #include "Engine/Engine.h"
 
+// Determining if we need to be disabled or not
+#include "Misc/ConfigCacheIni.h"
+
 // Steam tells us this number in documentation, however there's no define within the SDK
 #define STEAM_AUTH_MAX_TICKET_LENGTH_IN_BYTES 1024
 
@@ -20,7 +23,7 @@ FOnlineAuthSteam::FOnlineAuthSteam(FOnlineSubsystemSteam* InSubsystem) :
 	SteamUserPtr(SteamUser()),
 	SteamServerPtr(SteamGameServer()),
 	SteamSubsystem(InSubsystem),
-	bEnabled(true),
+	bEnabled(false),
 	bBadKey(false),
 	bReuseKey(false),
 	bBadWrite(false),
@@ -29,6 +32,18 @@ FOnlineAuthSteam::FOnlineAuthSteam(FOnlineSubsystemSteam* InSubsystem) :
 	bNeverSendKey(false),
 	bSendBadId(false)
 {
+	TArray<FString> ComponentList;
+	GConfig->GetArray(TEXT("PacketHandlerComponents"), TEXT("Components"), ComponentList, GEngineIni);
+
+	for (FString CompStr : ComponentList)
+	{
+		if (CompStr.Contains(TEXT("SteamAuthComponentModuleInterface")))
+		{
+			UE_LOG_ONLINE(Log, TEXT("AUTH: Steam Auth Enabled"));
+			bEnabled = true;
+			return;
+		}
+	}
 }
 
 FOnlineAuthSteam::FOnlineAuthSteam() :
