Fix typo
#rnx

diff --git a/Engine/Plugins/Online/OnlineSubsystemSteam/Source/Private/OnlineAuthInterfaceSteam.cpp b/Engine/Plugins/Online/OnlineSubsystemSteam/Source/Private/OnlineAuthInterfaceSteam.cpp
--- a/Engine/Plugins/Online/OnlineSubsystemSteam/Source/Private/OnlineAuthInterfaceSteam.cpp
+++ b/Engine/Plugins/Online/OnlineSubsystemSteam/Source/Private/OnlineAuthInterfaceSteam.cpp
@@ -319,7 +319,7 @@ bool FOnlineAuthSteam::KickPlayer(const FUniqueNetId& InUserId)
 
 		for (FConstPlayerControllerIterator Itr = World->GetPlayerControllerIterator(); Itr; ++Itr)
 		{
-			APlayerController* PC = It->Get();
+			APlayerController* PC = Itr->Get();
 			if (PC && PC->PlayerState != nullptr && PC->PlayerState->UniqueId.IsValid() &&
 				*(PC->PlayerState->UniqueId.GetUniqueNetId()) == InUserId)
 			{
