Copying //UE4/Dev-Audio to //UE4/Dev-Main (Source: //UE4/Dev-Audio @ 4332921)

============================
  MAJOR FEATURES & CHANGES
============================

Change 4233107 by [REDACTED_PERSON]

	New filter submix effect.

Change 4233108 by [REDACTED_PERSON]

	Adding dry-level for submix effect reverb.

Change 4233112 by [REDACTED_PERSON]

	Submix envelope following feature.

	- To use, register a BP delegate callback via AddEnvelopeFollowerDelegate.
	- Envelope value returned is a value for each channel in the submix (i.e. 2 for stereo, 6 for 5.1, etc)

Change 4309683 by [REDACTED_PERSON]

	Making unregistering a sound class happen on the audio thread.

Change 4309690 by [REDACTED_PERSON]

	Making unregistering a sound class happen on the audio thread.

	Fixing CL 4309683 typo

Change 4312378 by [REDACTED_PERSON]

	Checking in audio device hot swap fix code

Change 4312381 by [REDACTED_PERSON]

	Change label to Play When Silent for bVirtualizeWhenSilent

Change 4315516 by [REDACTED_PERSON]

	Fix for audio mixer source initialization state.

Change 4317052 by [REDACTED_PERSON]

	Unshelved from pending changelist '4017367':

	PR #4680: Refactor Voice Subsystems (Contributed by [REDACTED_PERSON])

Change 4317784 by [REDACTED_PERSON]

	New audio icons.

Change 4317785 by [REDACTED_PERSON]

	Hooking up new audio icons in C++.

	- Made preprocessor macros to help setting icons for audio assets
	- Moving icon definitions of other audio assets to AudioEditor module.

Change 4317787 by [REDACTED_PERSON]

	Fix for warning during cooking of source buses.

Change 4319727 by [REDACTED_PERSON]

	Updating icons w/ transparent backgrounds

Change 4320672 by [REDACTED_PERSON]

	Adding test map for Effects testing (and pan bug)

Change 4321142 by [REDACTED_PERSON]

	Making source buses not render as sound wave thumbnails

Change 4321143 by [REDACTED_PERSON]

	Fixing source buses so that new source buses will work after creating.

Change 4321191 by [REDACTED_PERSON]

	adding the Send_Test map to allow for crash repro

Change 4321456 by [REDACTED_PERSON]

	Fix for persistent submix effect settings between PIE sessions.

Change 4321503 by [REDACTED_PERSON]

	Fix for google resonance plugin when using -AudioPIEOnly and PIE shutdown.

	- Changed check to a warning.
	- Issue is when using AudioPIEOnly, decoding stream is never initialized...

Change 4321504 by [REDACTED_PERSON]

	Fixing TSharedPtr bug on PIE shutdown

	- Removing TSharedPtr for a lot of objects, which was causing object refs to leak
	- Moved to using TWeakPtr where it makes sense to keep the safety but avoid leaky refs
	- Fixing setting reset for both submix and source effect presets.

Change 4322887 by [REDACTED_PERSON]

	Fix for CIS error

Change 4324325 by [REDACTED_PERSON]

	fix for AudioMixerSubmix voice lookup crash #fyi aaron.mcleran #rb none

Change 4324340 by [REDACTED_PERSON]

	Come to think of it, it would also potentially crash here if we did not have a valid entry in the map. #fyi aaron.mcleran #rb none

Change 4324512 by [REDACTED_PERSON]

Change 4324668 by [REDACTED_PERSON]

	Removing unnecessary InitSoundSubmixes when creating source buses.

	- This function should only be called on engine init!

Change 4329745 by [REDACTED_PERSON]

	Fix InitSoundSubmix call to avoid a bad check on creating a USoundSubmix asset.

Change 4315524 by [REDACTED_PERSON]

	Making source effects be per-buffer vs per-frame.

	- Allows source effects to be more cache-coherent and opens door to SIMD-able source effects.
	- Also cleaned up DSP API to be more consistent and clear with respect to process functions
	- Optimized a few DSP objects (e.g. one-pole filters, etc)

diff --git a/Engine/Plugins/Online/OnlineSubsystem/Source/Public/Interfaces/VoiceInterface.h b/Engine/Plugins/Online/OnlineSubsystem/Source/Public/Interfaces/VoiceInterface.h
--- a/Engine/Plugins/Online/OnlineSubsystem/Source/Public/Interfaces/VoiceInterface.h
+++ b/Engine/Plugins/Online/OnlineSubsystem/Source/Public/Interfaces/VoiceInterface.h
@@ -243,19 +243,24 @@ class IOnlineVoice
 protected:
 	IOnlineVoice() {};
 
-	/**
-	 * Initialize the voice interface
-	 */
-	virtual bool Init() = 0;
 
 PACKAGE_SCOPE:
 
+	virtual IVoiceEnginePtr CreateVoiceEngine() = 0;
+
 	/**
 	 * Re-evaluates the muting list for all local talkers
 	 */
 	virtual void ProcessMuteChangeNotification() = 0;
 
 public:
+	virtual void Shutdown() {}
+
+	/**
+	* Initialize the voice interface
+	*/
+	virtual bool Init() { return true; }
+
 	/** Virtual destructor to force proper child cleanup */
 	virtual ~IOnlineVoice() {}
 
@@ -322,7 +327,7 @@ public:
 	virtual bool UnregisterRemoteTalker(const FUniqueNetId& UniqueId) = 0;
 
 	/**
-	 * Iterates the current remote talker list unregistering them with the 
+	 * Iterates the current remote talker list unregistering them with the
 	 * voice engine and our internal state
 	 */
 	virtual void RemoveAllRemoteTalkers() = 0;
diff --git a/Engine/Plugins/Online/OnlineSubsystemNull/Source/OnlineSubsystemNull.Build.cs b/Engine/Plugins/Online/OnlineSubsystemNull/Source/OnlineSubsystemNull.Build.cs
--- a/Engine/Plugins/Online/OnlineSubsystemNull/Source/OnlineSubsystemNull.Build.cs
+++ b/Engine/Plugins/Online/OnlineSubsystemNull/Source/OnlineSubsystemNull.Build.cs
@@ -10,14 +10,19 @@ public class OnlineSubsystemNull : ModuleRules
 		PrivateDefinitions.Add("ONLINESUBSYSTEMNULL_PACKAGE=1");
 		PCHUsage = ModuleRules.PCHUsageMode.UseExplicitOrSharedPCHs;
 
-		PrivateDependencyModuleNames.AddRange(
+        PublicDependencyModuleNames.AddRange(
+            new string[] {
+                "OnlineSubsystemUtils"
+            }
+            );
+
+        PrivateDependencyModuleNames.AddRange(
 			new string[] {
 				"Core", 
 				"CoreUObject", 
 				"Engine", 
 				"Sockets", 
 				"OnlineSubsystem", 
-				"OnlineSubsystemUtils",
 				"Json"
 			}
 			);
diff --git a/Engine/Plugins/Online/OnlineSubsystemNull/Source/Private/OnlineSubsystemNull.cpp b/Engine/Plugins/Online/OnlineSubsystemNull/Source/Private/OnlineSubsystemNull.cpp
--- a/Engine/Plugins/Online/OnlineSubsystemNull/Source/Private/OnlineSubsystemNull.cpp
+++ b/Engine/Plugins/Online/OnlineSubsystemNull/Source/Private/OnlineSubsystemNull.cpp
@@ -7,7 +7,7 @@
 #include "OnlineSessionInterfaceNull.h"
 #include "OnlineLeaderboardInterfaceNull.h"
 #include "OnlineIdentityNull.h"
-#include "VoiceInterfaceImpl.h"
+#include "VoiceInterfaceNull.h"
 #include "OnlineAchievementsInterfaceNull.h"
 
 FThreadSafeCounter FOnlineSubsystemNull::TaskCounter;
diff --git a/Engine/Plugins/Online/OnlineSubsystemNull/Source/Private/VoiceInterfaceNull.h b/Engine/Plugins/Online/OnlineSubsystemNull/Source/Private/VoiceInterfaceNull.h
new file mode 100644
--- /dev/null
+++ b/Engine/Plugins/Online/OnlineSubsystemNull/Source/Private/VoiceInterfaceNull.h
@@ -0,0 +1,25 @@
+// Copyright 1998-2018 Epic Games, Inc. All Rights Reserved.
+
+#include "VoiceInterfaceImpl.h"
+#include "UObject/CoreOnline.h"
+#include "OnlineSubsystemTypes.h"
+#include "Interfaces/VoiceInterface.h"
+#include "Net/VoiceDataCommon.h"
+#include "OnlineSubsystemUtilsPackage.h"
+
+class ONLINESUBSYSTEMNULL_API FOnlineVoiceImplNull : public FOnlineVoiceImpl {
+PACKAGE_SCOPE:
+	FOnlineVoiceImplNull() : FOnlineVoiceImpl()
+	{};
+
+public:
+	/** Constructor */
+	FOnlineVoiceImplNull(class IOnlineSubsystem* InOnlineSubsystem) :
+		FOnlineVoiceImpl(InOnlineSubsystem)
+	{
+	};
+
+	virtual ~FOnlineVoiceImplNull() override {}
+};
+
+typedef TSharedPtr<FOnlineVoiceImpl, ESPMode::ThreadSafe> FOnlineVoiceImplPtr;
diff --git a/Engine/Plugins/Online/OnlineSubsystemNull/Source/Public/OnlineSubsystemNull.h b/Engine/Plugins/Online/OnlineSubsystemNull/Source/Public/OnlineSubsystemNull.h
--- a/Engine/Plugins/Online/OnlineSubsystemNull/Source/Public/OnlineSubsystemNull.h
+++ b/Engine/Plugins/Online/OnlineSubsystemNull/Source/Public/OnlineSubsystemNull.h
@@ -109,7 +109,7 @@ private:
 	FOnlineSessionNullPtr SessionInterface;
 
 	/** Interface for voice communication */
-	mutable FOnlineVoiceImplPtr VoiceInterface;
+	mutable IOnlineVoicePtr VoiceInterface;
 
 	/** Interface for voice communication */
 	mutable bool bVoiceInterfaceInitialized;
diff --git a/Engine/Plugins/Online/OnlineSubsystemOculus/Source/Private/OnlineVoiceOculus.h b/Engine/Plugins/Online/OnlineSubsystemOculus/Source/Private/OnlineVoiceOculus.h
--- a/Engine/Plugins/Online/OnlineSubsystemOculus/Source/Private/OnlineVoiceOculus.h
+++ b/Engine/Plugins/Online/OnlineSubsystemOculus/Source/Private/OnlineVoiceOculus.h
@@ -72,6 +72,7 @@ public:
 	{
 		return (!bIsLocalPlayerMuted && LocalUserNum <= MAX_LOCAL_PLAYERS);
 	}
+
 	virtual bool IsRemotePlayerTalking(const FUniqueNetId& UniqueId) override;
 	virtual bool IsMuted(uint32 LocalUserNum, const FUniqueNetId& UniqueId) const override;
 	virtual bool MuteRemoteTalker(uint8 LocalUserNum, const FUniqueNetId& PlayerId, bool bIsSystemWide) override;
@@ -109,6 +110,11 @@ PACKAGE_SCOPE:
 	virtual bool Init() override;
 	void ProcessMuteChangeNotification() override {}
 
+	IVoiceEnginePtr CreateVoiceEngine() override
+	{
+		return IVoiceEnginePtr();
+	}
+
 	/**
 	* Submits network packets to audio system for playback
 	*/
diff --git a/Engine/Plugins/Online/OnlineSubsystemSteam/Source/OnlineSubsystemSteam.Build.cs b/Engine/Plugins/Online/OnlineSubsystemSteam/Source/OnlineSubsystemSteam.Build.cs
--- a/Engine/Plugins/Online/OnlineSubsystemSteam/Source/OnlineSubsystemSteam.Build.cs
+++ b/Engine/Plugins/Online/OnlineSubsystemSteam/Source/OnlineSubsystemSteam.Build.cs
@@ -29,6 +29,7 @@ public class OnlineSubsystemSteam : ModuleRules
 				"Engine", 
 				"Sockets", 
 				"Voice",
+                "AudioMixer",
 				"OnlineSubsystem",
 				"Json",
 				"PacketHandler"
diff --git a/Engine/Plugins/Online/OnlineSubsystemSteam/Source/Private/OnlineSubsystemSteam.cpp b/Engine/Plugins/Online/OnlineSubsystemSteam/Source/Private/OnlineSubsystemSteam.cpp
--- a/Engine/Plugins/Online/OnlineSubsystemSteam/Source/Private/OnlineSubsystemSteam.cpp
+++ b/Engine/Plugins/Online/OnlineSubsystemSteam/Source/Private/OnlineSubsystemSteam.cpp
@@ -29,6 +29,7 @@
 #include "OnlineExternalUIInterfaceSteam.h"
 #include "OnlineAchievementsInterfaceSteam.h"
 #include "OnlineAuthInterfaceSteam.h"
+#include "VoiceInterfaceSteam.h"
 
 
 #if !UE_BUILD_SHIPPING
@@ -234,6 +235,16 @@ IOnlineLeaderboardsPtr FOnlineSubsystemSteam::GetLeaderboardsInterface() const
 
 IOnlineVoicePtr FOnlineSubsystemSteam::GetVoiceInterface() const
 {
+	if (VoiceInterface.IsValid() && !bVoiceInterfaceInitialized)
+	{
+		if (!VoiceInterface->Init())
+		{
+			VoiceInterface = nullptr;
+		}
+
+		bVoiceInterfaceInitialized = true;
+	}
+
 	return VoiceInterface;
 }
 
@@ -341,7 +352,7 @@ bool FOnlineSubsystemSteam::Tick(float DeltaTime)
 		SessionInterface->Tick(DeltaTime);
 	}
 
-	if (VoiceInterface.IsValid())
+	if (VoiceInterface.IsValid() && bVoiceInterfaceInitialized)
 	{
 		VoiceInterface->Tick(DeltaTime);
 	}
@@ -403,10 +414,6 @@ bool FOnlineSubsystemSteam::Init()
 			SharedCloudInterface = MakeShareable(new FOnlineSharedCloudSteam(this));
 			LeaderboardsInterface = MakeShareable(new FOnlineLeaderboardsSteam(this));
 			VoiceInterface = MakeShareable(new FOnlineVoiceSteam(this));
-			if (!VoiceInterface->Init())
-			{
-				VoiceInterface = nullptr;
-			}
 			ExternalUIInterface = MakeShareable(new FOnlineExternalUISteam(this));
 			AchievementsInterface = MakeShareable(new FOnlineAchievementsSteam(this));
 
@@ -450,7 +457,11 @@ bool FOnlineSubsystemSteam::Shutdown()
 		delete OnlineAsyncTaskThreadRunnable;
 		OnlineAsyncTaskThreadRunnable = nullptr;
 	}
-	
+
+	if (VoiceInterface.IsValid() && bVoiceInterfaceInitialized) {
+		VoiceInterface->Shutdown();
+	}
+
 #define DESTRUCT_INTERFACE(Interface) \
 	if (Interface.IsValid()) \
 	{ \
diff --git a/Engine/Plugins/Online/OnlineSubsystemSteam/Source/Private/VoiceEngineSteam.cpp b/Engine/Plugins/Online/OnlineSubsystemSteam/Source/Private/VoiceEngineSteam.cpp
--- a/Engine/Plugins/Online/OnlineSubsystemSteam/Source/Private/VoiceEngineSteam.cpp
+++ b/Engine/Plugins/Online/OnlineSubsystemSteam/Source/Private/VoiceEngineSteam.cpp
@@ -1,48 +1,141 @@
 // Copyright 1998-2018 Epic Games, Inc. All Rights Reserved.
 
 #include "VoiceEngineSteam.h"
-#include "Misc/ConfigCacheIni.h"
 #include "Components/AudioComponent.h"
+#include "GameFramework/GameSession.h"
 #include "OnlineSubsystemSteam.h"
+#include "OnlineSubsystemUtils.h"
+#include "Sound/SoundWaveProcedural.h"
 #include "VoiceModule.h"
+#include "Voice.h"
 #include "SteamUtilities.h"
 
-#include "Sound/SoundWaveProcedural.h"
-#include "OnlineSubsystemUtils.h"
+/** Largest size allowed to carry over into next buffer */
+#define MAX_VOICE_REMAINDER_SIZE 4 * 1024
+
+FRemoteTalkerDataSteam::FRemoteTalkerDataSteam() :
+	MaxUncompressedDataSize(0),
+	MaxUncompressedDataQueueSize(0),
+	CurrentUncompressedDataQueueSize(0),
+	LastSeen(0.0),
+	NumFramesStarved(0),
+	VoipSynthComponent(nullptr),
+	VoiceDecoder(nullptr)
+{
+	int32 SampleRate = UVOIPStatics::GetVoiceSampleRate();
+	int32 NumChannels = DEFAULT_NUM_VOICE_CHANNELS;
+	VoiceDecoder = FVoiceModule::Get().CreateVoiceDecoder(SampleRate, NumChannels);
+	check(VoiceDecoder.IsValid());
+
+	// Approx 1 sec worth of data for a stereo microphone
+	MaxUncompressedDataSize = UVOIPStatics::GetMaxUncompressedVoiceDataSizePerChannel() * 2;
+	MaxUncompressedDataQueueSize = MaxUncompressedDataSize * 5;
+	{
+		FScopeLock ScopeLock(&QueueLock);
+		UncompressedDataQueue.Empty(MaxUncompressedDataQueueSize);
+	}
+}
+
+FRemoteTalkerDataSteam::FRemoteTalkerDataSteam(FRemoteTalkerDataSteam&& Other)
+{
+	LastSeen = Other.LastSeen;
+	Other.LastSeen = 0.0;
+
+	NumFramesStarved = Other.NumFramesStarved;
+	Other.NumFramesStarved = 0;
+
+	VoipSynthComponent = Other.VoipSynthComponent;
+	Other.VoipSynthComponent = nullptr;
 
-/** Largest size Steam says it will need to compress data */
-#define MAX_COMPRESSED_VOICE_BUFFER_SIZE 8 * 1024
-/** Largest size Steam says it will uncompress to */
-#define MAX_UNCOMPRESSED_VOICE_BUFFER_SIZE 22 * 1024
+	VoiceDecoder = MoveTemp(Other.VoiceDecoder);
+	Other.VoiceDecoder = nullptr;
 
-FVoiceEngineSteam::FVoiceEngineSteam(FOnlineSubsystemSteam* InSteamSubsystem) :
-	SteamSubsystem(InSteamSubsystem),
+	MaxUncompressedDataSize = Other.MaxUncompressedDataSize;
+	Other.MaxUncompressedDataSize = 0;
+
+	MaxUncompressedDataQueueSize = Other.MaxUncompressedDataQueueSize;
+	Other.MaxUncompressedDataQueueSize = 0;
+
+	CurrentUncompressedDataQueueSize = Other.CurrentUncompressedDataQueueSize;
+	Other.CurrentUncompressedDataQueueSize = 0;
+
+	{
+		FScopeLock ScopeLock(&Other.QueueLock);
+		UncompressedDataQueue = MoveTemp(Other.UncompressedDataQueue);
+	}
+}
+
+FRemoteTalkerDataSteam::~FRemoteTalkerDataSteam()
+{
+	VoiceDecoder = nullptr;
+
+	CurrentUncompressedDataQueueSize = 0;
+
+	{
+		FScopeLock ScopeLock(&QueueLock);
+		UncompressedDataQueue.Empty();
+	}
+}
+
+void FRemoteTalkerDataSteam::Reset()
+{
+	// Set to large number so TickTalkers doesn't come in here
+	LastSeen = MAX_FLT;
+	NumFramesStarved = 0;
+
+	if (VoipSynthComponent)
+	{
+		VoipSynthComponent->Stop();
+		bIsActive = false;
+	}
+
+	CurrentUncompressedDataQueueSize = 0;
+
+	{
+		FScopeLock ScopeLock(&QueueLock);
+		UncompressedDataQueue.Empty();
+	}
+}
+
+void FRemoteTalkerDataSteam::Cleanup()
+{
+	if (VoipSynthComponent)
+	{
+		VoipSynthComponent->Stop();
+		bIsActive = false;
+	}
+
+	VoipSynthComponent = nullptr;
+}
+
+FVoiceEngineSteam::FVoiceEngineSteam(IOnlineSubsystem* InSubsystem) :
+	OnlineSubsystem(InSubsystem),
+	VoiceCapture(nullptr),
+	VoiceEncoder(nullptr),
 	OwningUserIndex(INVALID_INDEX),
+	UncompressedBytesAvailable(0),
+	CompressedBytesAvailable(0),
+	AvailableVoiceResult(EVoiceCaptureState::UnInitialized),
 	bPendingFinalCapture(false),
 	bIsCapturing(false),
-	SerializeHelper(NULL)
+	SerializeHelper(nullptr)
 {
 	SteamUserPtr = SteamUser();
 	SteamFriendsPtr = SteamFriends();
-
-	FCoreUObjectDelegates::PostLoadMapWithWorld.AddRaw(this, &FVoiceEngineSteam::OnPostLoadMap);
 }
 
 FVoiceEngineSteam::~FVoiceEngineSteam()
 {
 	if (bIsCapturing)
 	{
-		SteamUserPtr->StopVoiceRecording();
+		VoiceCapture->Stop();
 		SteamFriendsPtr->SetInGameVoiceSpeaking(SteamUserPtr->GetSteamID(), false);
 	}
 
-	if (SerializeHelper != nullptr)
-	{
-		delete SerializeHelper;
-		SerializeHelper = nullptr;
-	}
+	VoiceCapture = nullptr;
+	VoiceEncoder = nullptr;
 
-	FCoreUObjectDelegates::PostLoadMapWithWorld.RemoveAll(this);
+	delete SerializeHelper;
 }
 
 void FVoiceEngineSteam::VoiceCaptureUpdate() const
@@ -50,10 +143,10 @@ void FVoiceEngineSteam::VoiceCaptureUpdate() const
 	if (bPendingFinalCapture)
 	{
 		uint32 CompressedSize;
-		const EVoiceResult RecordingState = SteamUserPtr->GetAvailableVoice(&CompressedSize, NULL, 0);
+		const EVoiceCaptureState::Type RecordingState = VoiceCapture->GetCaptureState(CompressedSize);
 
 		// If no data is available, we have finished capture the last (post-StopRecording) half-second of voice data
-		if (RecordingState == k_EVoiceResultNotRecording)
+		if (RecordingState == EVoiceCaptureState::NotCapturing)
 		{
 			UE_LOG(LogVoiceEngine, Log, TEXT("Internal voice capture complete."));
 
@@ -76,19 +169,24 @@ void FVoiceEngineSteam::VoiceCaptureUpdate() const
 void FVoiceEngineSteam::StartRecording() const
 {
 	UE_LOG(LogVoiceEngine, VeryVerbose, TEXT("VOIP StartRecording"));
-	if (SteamUserPtr)
+	if (VoiceCapture.IsValid())
 	{
-		SteamUserPtr->StartVoiceRecording();
-		SteamFriendsPtr->SetInGameVoiceSpeaking(SteamUserPtr->GetSteamID(), true);
+		if (!VoiceCapture->Start())
+		{
+			UE_LOG(LogVoiceEngine, Warning, TEXT("Failed to start voice recording"));
+		}
+		else if (SteamFriendsPtr) {
+			SteamFriendsPtr->SetInGameVoiceSpeaking(SteamUserPtr->GetSteamID(), true);
+		}
 	}
 }
 
 void FVoiceEngineSteam::StopRecording() const
 {
 	UE_LOG(LogVoiceEngine, VeryVerbose, TEXT("VOIP StopRecording"));
-	if (SteamUserPtr)
+	if (VoiceCapture.IsValid())
 	{
-		SteamUserPtr->StopVoiceRecording();
+		VoiceCapture->Stop();
 	}
 }
 
@@ -101,45 +199,45 @@ void FVoiceEngineSteam::StoppedRecording() const
 	}
 }
 
-bool FVoiceEngineSteam::Init(int32 MaxLocalTalkers, int32 MaxRemoteTalkers) 
+bool FVoiceEngineSteam::Init(int32 MaxLocalTalkers, int32 MaxRemoteTalkers)
 {
 	bool bSuccess = false;
 
-	if (SteamSubsystem->IsSteamClientAvailable())
+	if (!OnlineSubsystem->IsDedicated())
 	{
-		bool bHasVoiceEnabled = false;
-		if (GConfig->GetBool(TEXT("OnlineSubsystem"),TEXT("bHasVoiceEnabled"), bHasVoiceEnabled, GEngineIni) && bHasVoiceEnabled)
+		FVoiceModule& VoiceModule = FVoiceModule::Get();
+		if (VoiceModule.IsVoiceEnabled())
 		{
-			if (SteamUserPtr && SteamFriendsPtr)
+			VoiceCapture = VoiceModule.CreateVoiceCapture();
+			VoiceEncoder = VoiceModule.CreateVoiceEncoder();
+
+			bSuccess = VoiceCapture.IsValid() && VoiceEncoder.IsValid();
+			if (bSuccess)
 			{
-				// Just verify voice capture is available
-				uint32 CompressedBytes;
-				SteamUserPtr->StartVoiceRecording();
-				EVoiceResult VoiceResult = SteamUserPtr->GetAvailableVoice(&CompressedBytes, NULL, 0);
-				SteamUserPtr->StopVoiceRecording();
-
-				bSuccess = (VoiceResult != k_EVoiceResultNotInitialized);
-				if (bSuccess)
-				{
-					CompressedVoiceBuffer.Empty(MAX_COMPRESSED_VOICE_BUFFER_SIZE);
-					DecompressedVoiceBuffer.Empty(MAX_UNCOMPRESSED_VOICE_BUFFER_SIZE);
-				}
-				else
+				CompressedVoiceBuffer.Empty(UVOIPStatics::GetMaxCompressedVoiceDataSize());
+				DecompressedVoiceBuffer.Empty(UVOIPStatics::GetMaxUncompressedVoiceDataSizePerChannel());
+
+				for (int32 TalkerIdx = 0; TalkerIdx < MaxLocalTalkers; TalkerIdx++)
 				{
-					UE_LOG(LogVoice, Warning, TEXT("Steamworks voice initialization failed!"));
+					PlayerVoiceData[TalkerIdx].VoiceRemainderSize = 0;
+					PlayerVoiceData[TalkerIdx].VoiceRemainder.Empty(MAX_VOICE_REMAINDER_SIZE);
 				}
 			}
+			else
+			{
+				UE_LOG(LogVoice, Warning, TEXT("Voice capture initialization failed!"));
+			}
 		}
 		else
 		{
-			UE_LOG(LogVoice, Log, TEXT("Voice module disabled by config [OnlineSubsystem].bHasVoiceEnabled"));
+			UE_LOG(LogVoice, Log, TEXT("Voice module disabled by config [Voice].bEnabled"));
 		}
 	}
 
 	return bSuccess;
 }
 
-uint32 FVoiceEngineSteam::StartLocalVoiceProcessing(uint32 LocalUserNum) 
+uint32 FVoiceEngineSteam::StartLocalVoiceProcessing(uint32 LocalUserNum)
 {
 	uint32 Return = E_FAIL;
 	if (IsOwningUser(LocalUserNum))
@@ -167,7 +265,7 @@ uint32 FVoiceEngineSteam::StartLocalVoiceProcessing(uint32 LocalUserNum)
 	return Return;
 }
 
-uint32 FVoiceEngineSteam::StopLocalVoiceProcessing(uint32 LocalUserNum) 
+uint32 FVoiceEngineSteam::StopLocalVoiceProcessing(uint32 LocalUserNum)
 {
 	uint32 Return = E_FAIL;
 	if (IsOwningUser(LocalUserNum))
@@ -194,14 +292,27 @@ uint32 FVoiceEngineSteam::StopLocalVoiceProcessing(uint32 LocalUserNum)
 	return Return;
 }
 
-uint32 FVoiceEngineSteam::GetVoiceDataReadyFlags() const 
+uint32 FVoiceEngineSteam::UnregisterRemoteTalker(const FUniqueNetId& UniqueId)
+{
+	FRemoteTalkerDataSteam* RemoteData = RemoteTalkerBuffers.Find(FUniqueNetIdWrapper(UniqueId.AsShared()));
+	if (RemoteData != nullptr)
+	{
+		// Dump the whole talker
+		RemoteData->Cleanup();
+		RemoteTalkerBuffers.Remove(FUniqueNetIdWrapper(UniqueId.AsShared()));
+	}
+
+	return S_OK;
+}
+
+uint32 FVoiceEngineSteam::GetVoiceDataReadyFlags() const
 {
 	// First check and update the internal state of VOIP recording
 	VoiceCaptureUpdate();
 	if (OwningUserIndex != INVALID_INDEX && IsRecording())
 	{
-		// Check if there is new data available via the Steam API
-		if (AvailableVoiceResult == k_EVoiceResultOK && CompressedBytesAvailable > 0)
+		// Check if there is new data available via the Voice API
+		if (AvailableVoiceResult == EVoiceCaptureState::Ok && UncompressedBytesAvailable > 0)
 		{
 			return 1 << OwningUserIndex;
 		}
@@ -210,102 +321,135 @@ uint32 FVoiceEngineSteam::GetVoiceDataReadyFlags() const
 	return 0;
 }
 
-uint32 FVoiceEngineSteam::ReadLocalVoiceData(uint32 LocalUserNum, uint8* Data, uint32* Size) 
+uint32 FVoiceEngineSteam::ReadLocalVoiceData(uint32 LocalUserNum, uint8* Data, uint32* Size, uint64* OutSampleCount)
 {
 	check(*Size > 0);
 
 	// Before doing anything, check/update the current recording state
 	VoiceCaptureUpdate();
 
-	// Return data even if not capturing, if the final half-second of data from Steam is still pending
+	// Return data even if not capturing, possibly have data during stopping
 	if (IsOwningUser(LocalUserNum) && IsRecording())
 	{
-		CompressedVoiceBuffer.Empty(MAX_COMPRESSED_VOICE_BUFFER_SIZE);
+		DecompressedVoiceBuffer.Empty(UVOIPStatics::GetMaxUncompressedVoiceDataSizePerChannel());
+		CompressedVoiceBuffer.Empty(UVOIPStatics::GetMaxCompressedVoiceDataSize());
 
-		uint32 CompressedBytes = 0;
-		EVoiceResult VoiceResult = SteamUserPtr->GetAvailableVoice(&CompressedBytes, NULL, 0);
-		if (VoiceResult != k_EVoiceResultOK && VoiceResult != k_EVoiceResultNoData)
+		uint32 NewVoiceDataBytes = 0;
+		EVoiceCaptureState::Type VoiceResult = VoiceCapture->GetCaptureState(NewVoiceDataBytes);
+		if (VoiceResult != EVoiceCaptureState::Ok && VoiceResult != EVoiceCaptureState::NoData)
 		{
-			UE_LOG(LogVoiceEngine, Warning, TEXT("ReadLocalVoiceData: GetAvailableVoice failure: VoiceResult: %s"), *SteamVoiceResult(VoiceResult));
+			UE_LOG(LogVoiceEngine, Warning, TEXT("ReadLocalVoiceData: GetAvailableVoice failure: VoiceResult: %s"), EVoiceCaptureState::ToString(VoiceResult));
 			return E_FAIL;
 		}
 
-		// This shouldn't happen, but just in case.
-		if (CompressedBytes == 0)
+		if (NewVoiceDataBytes == 0)
 		{
-			UE_LOG(LogVoiceEngine, VeryVerbose, TEXT("ReadLocalVoiceData: No Data: VoiceResult: %s"), *SteamVoiceResult(VoiceResult));
+			UE_LOG(LogVoiceEngine, VeryVerbose, TEXT("ReadLocalVoiceData: No Data: VoiceResult: %s"), EVoiceCaptureState::ToString(VoiceResult));
 			*Size = 0;
 			return S_OK;
 		}
 
-		// Update the amount of data available for consumption (CompressedBytes can be more than AvailableWritten below)
-		CompressedVoiceBuffer.AddUninitialized(CompressedBytes);
+		// Make space for new and any previously remaining data
 
-		uint32 AvailableWritten = 0;
-		VoiceResult = SteamUserPtr->GetVoice(true, CompressedVoiceBuffer.GetData(), CompressedVoiceBuffer.Num(), &AvailableWritten, false, NULL, 0, NULL, 0);
+		// Add the number of new bytes (since last time this function was called) and the number of bytes remaining that wasn't consumed last time this was called
+		// This is how many bytes we would like to return
+		uint32 TotalVoiceBytes = NewVoiceDataBytes + PlayerVoiceData[LocalUserNum].VoiceRemainderSize;
 
-		static double LastGetVoiceCallTime = 0.0;
-		double CurTime = FPlatformTime::Seconds();
-		double TimeSinceLastCall = (LastGetVoiceCallTime > 0) ? (CurTime - LastGetVoiceCallTime) : 0.0;
-		LastGetVoiceCallTime = CurTime;
-
-		UE_LOG(LogVoiceEngine, Log, TEXT("ReadLocalVoiceData: GetVoice: Result: %s, Available: %i, LastCall: %0.3f"), *SteamVoiceResult(VoiceResult), AvailableWritten, TimeSinceLastCall * 1000.0);
-		if (VoiceResult != k_EVoiceResultOK)
+		// But we have a max amount we can return so clamp it to that max value if we're asking for more bytes than we're allowed
+		if (TotalVoiceBytes > UVOIPStatics::GetMaxUncompressedVoiceDataSizePerChannel())
 		{
-			UE_LOG(LogVoiceEngine, Warning, TEXT("ReadLocalVoiceData: GetVoice failure: VoiceResult: %s"), *SteamVoiceResult(VoiceResult));
-			*Size = 0;
-			CompressedVoiceBuffer.Empty(MAX_COMPRESSED_VOICE_BUFFER_SIZE);
-			return E_FAIL;
+			UE_LOG(LogVoiceEngine, Warning, TEXT("Exceeded uncompressed voice buffer size, clamping"))
+				TotalVoiceBytes = UVOIPStatics::GetMaxUncompressedVoiceDataSizePerChannel();
 		}
 
-		if (AvailableWritten > 0)
-		{
-			*Size = FMath::Min<int32>(*Size, AvailableWritten);
-			FMemory::Memcpy(Data, CompressedVoiceBuffer.GetData(), *Size);
+		DecompressedVoiceBuffer.AddUninitialized(TotalVoiceBytes);
 
-			UE_LOG(LogVoiceEngine, VeryVerbose, TEXT("ReadLocalVoiceData: Size: %d"), *Size);
-			return S_OK;
+		// If there's still audio left from a previous ReadLocalData call that didn't get output, copy that first into the decompressed voice buffer
+		if (PlayerVoiceData[LocalUserNum].VoiceRemainderSize > 0)
+		{
+			FMemory::Memcpy(DecompressedVoiceBuffer.GetData(), PlayerVoiceData[LocalUserNum].VoiceRemainder.GetData(), PlayerVoiceData[LocalUserNum].VoiceRemainderSize);
 		}
-		else
+
+		// Get new uncompressed data
+		uint8* RemainingDecompressedBufferPtr = DecompressedVoiceBuffer.GetData() + PlayerVoiceData[LocalUserNum].VoiceRemainderSize;
+		uint32 ByteWritten = 0;
+		uint64 NewSampleCount = 0;
+		VoiceResult = VoiceCapture->GetVoiceData(DecompressedVoiceBuffer.GetData() + PlayerVoiceData[LocalUserNum].VoiceRemainderSize, NewVoiceDataBytes, ByteWritten, NewSampleCount);
+
+		TotalVoiceBytes = ByteWritten + PlayerVoiceData[LocalUserNum].VoiceRemainderSize;
+
+		if ((VoiceResult == EVoiceCaptureState::Ok || VoiceResult == EVoiceCaptureState::NoData) && TotalVoiceBytes > 0)
 		{
-			*Size = 0;
+			if (OutSampleCount != nullptr)
+			{
+				*OutSampleCount = NewSampleCount;
+			}
+
+			// Prepare the encoded buffer (e.g. opus)
+			CompressedBytesAvailable = UVOIPStatics::GetMaxCompressedVoiceDataSize();
+			CompressedVoiceBuffer.AddUninitialized(UVOIPStatics::GetMaxCompressedVoiceDataSize());
+
+			check(((uint32)CompressedVoiceBuffer.Num()) <= UVOIPStatics::GetMaxCompressedVoiceDataSize());
+
+			// Run the uncompressed audio through the opus decoder, note that it may not encode all data, which results in some remaining data
+			PlayerVoiceData[LocalUserNum].VoiceRemainderSize =
+				VoiceEncoder->Encode(DecompressedVoiceBuffer.GetData(), TotalVoiceBytes, CompressedVoiceBuffer.GetData(), CompressedBytesAvailable);
+
+			// Save off any unencoded remainder
+			if (PlayerVoiceData[LocalUserNum].VoiceRemainderSize > 0)
+			{
+				if (PlayerVoiceData[LocalUserNum].VoiceRemainderSize > MAX_VOICE_REMAINDER_SIZE)
+				{
+					UE_LOG(LogVoiceEngine, Warning, TEXT("Exceeded voice remainder buffer size, clamping"));
+					PlayerVoiceData[LocalUserNum].VoiceRemainderSize = MAX_VOICE_REMAINDER_SIZE;
+				}
+
+				PlayerVoiceData[LocalUserNum].VoiceRemainder.AddUninitialized(MAX_VOICE_REMAINDER_SIZE);
+				FMemory::Memcpy(PlayerVoiceData[LocalUserNum].VoiceRemainder.GetData(), DecompressedVoiceBuffer.GetData() + (TotalVoiceBytes - PlayerVoiceData[LocalUserNum].VoiceRemainderSize), PlayerVoiceData[LocalUserNum].VoiceRemainderSize);
+			}
+
+			static double LastGetVoiceCallTime = 0.0;
+			double CurTime = FPlatformTime::Seconds();
+			double TimeSinceLastCall = (LastGetVoiceCallTime > 0) ? (CurTime - LastGetVoiceCallTime) : 0.0;
+			LastGetVoiceCallTime = CurTime;
+
+			UE_LOG(LogVoiceEngine, Log, TEXT("ReadLocalVoiceData: GetVoice: Result: %s, Available: %i, LastCall: %0.3f ms"), EVoiceCaptureState::ToString(VoiceResult), CompressedBytesAvailable, TimeSinceLastCall * 1000.0);
+			if (CompressedBytesAvailable > 0)
+			{
+				*Size = FMath::Min<int32>(*Size, CompressedBytesAvailable);
+				FMemory::Memcpy(Data, CompressedVoiceBuffer.GetData(), *Size);
+
+				UE_LOG(LogVoiceEngine, VeryVerbose, TEXT("ReadLocalVoiceData: Size: %d"), *Size);
+				return S_OK;
+			}
+			else
+			{
+				*Size = 0;
+				CompressedVoiceBuffer.Empty(UVOIPStatics::GetMaxCompressedVoiceDataSize());
+
+				UE_LOG(LogVoiceEngine, Warning, TEXT("ReadLocalVoiceData: GetVoice failure: VoiceResult: %s"), EVoiceCaptureState::ToString(VoiceResult));
+				return E_FAIL;
+			}
 		}
 	}
 
 	return E_FAIL;
 }
 
-uint32 FVoiceEngineSteam::SubmitRemoteVoiceData(const FUniqueNetId& RemoteTalkerId, uint8* Data, uint32* Size) 
+uint32 FVoiceEngineSteam::SubmitRemoteVoiceData(const FUniqueNetIdWrapper& RemoteTalkerId, uint8* Data, uint32* Size, uint64& InSampleCount)
 {
 	UE_LOG(LogVoiceEngine, VeryVerbose, TEXT("SubmitRemoteVoiceData(%s) Size: %d received!"), *RemoteTalkerId.ToDebugString(), *Size);
 
-	const FUniqueNetIdSteam& SteamId = (const FUniqueNetIdSteam&)RemoteTalkerId;
-	FRemoteTalkerDataSteam* QueuedData = RemoteTalkerBuffers.Find(SteamId);
-
-	if (QueuedData == NULL)
-	{
-		RemoteTalkerBuffers.Add(SteamId, FRemoteTalkerDataSteam());
-		QueuedData = RemoteTalkerBuffers.Find(SteamId);
-	}
-
-	check(QueuedData);
+	FRemoteTalkerDataSteam& QueuedData = RemoteTalkerBuffers.FindOrAdd(RemoteTalkerId);
 
 	// new voice packet.
-	QueuedData->LastSeen = FPlatformTime::Seconds();
-
-	uint32 BytesWritten = 0;
+	QueuedData.LastSeen = FPlatformTime::Seconds();
 
-	DecompressedVoiceBuffer.Empty(MAX_UNCOMPRESSED_VOICE_BUFFER_SIZE);
-	DecompressedVoiceBuffer.AddUninitialized(MAX_UNCOMPRESSED_VOICE_BUFFER_SIZE);
-	const EVoiceResult VoiceResult = SteamUserPtr->DecompressVoice(Data, *Size, DecompressedVoiceBuffer.GetData(),
-		DecompressedVoiceBuffer.Num(), &BytesWritten, SteamUserPtr->GetVoiceOptimalSampleRate());
+	uint32 BytesWritten = UVOIPStatics::GetMaxUncompressedVoiceDataSizePerChannel();
 
-	if (VoiceResult != k_EVoiceResultOK)
-	{
-		UE_LOG(LogVoiceEngine, Warning, TEXT("SubmitRemoteVoiceData: DecompressVoice failure: VoiceResult: %s"), *SteamVoiceResult(VoiceResult));
-		*Size = 0;
-		return E_FAIL;
-	}
+	DecompressedVoiceBuffer.Empty(UVOIPStatics::GetMaxUncompressedVoiceDataSizePerChannel());
+	DecompressedVoiceBuffer.AddUninitialized(UVOIPStatics::GetMaxUncompressedVoiceDataSizePerChannel());
+	QueuedData.VoiceDecoder->Decode(Data, *Size, DecompressedVoiceBuffer.GetData(), BytesWritten);
 
 	// If there is no data, return
 	if (BytesWritten <= 0)
@@ -314,30 +458,54 @@ uint32 FVoiceEngineSteam::SubmitRemoteVoiceData(const FUniqueNetId& RemoteTalker
 		return S_OK;
 	}
 
+	bool bAudioComponentCreated = false;
 	// Generate a streaming wave audio component for voice playback
-	if (QueuedData->AudioComponent == NULL || QueuedData->AudioComponent->IsPendingKill())
+	if (QueuedData.VoipSynthComponent == nullptr || QueuedData.VoipSynthComponent->IsPendingKill())
 	{
-		if (SerializeHelper == NULL)
+		if (SerializeHelper == nullptr)
 		{
 			SerializeHelper = new FVoiceSerializeHelper(this);
 		}
 
-		QueuedData->AudioComponent = CreateVoiceAudioComponent(SteamUserPtr->GetVoiceOptimalSampleRate(), DEFAULT_NUM_VOICE_CHANNELS);
-		if (QueuedData->AudioComponent)
+		QueuedData.VoipSynthComponent = CreateVoiceSynthComponent(UVOIPStatics::GetVoiceSampleRate());
+		if (QueuedData.VoipSynthComponent)
 		{
-			QueuedData->AudioComponent->OnAudioFinishedNative.AddRaw(this, &FVoiceEngineSteam::OnAudioFinished);
-			QueuedData->AudioComponent->Play();
+			//TODO, make buffer size and buffering delay runtime-controllable parameters.
+			QueuedData.bIsActive = false;
+			QueuedData.VoipSynthComponent->OpenPacketStream(InSampleCount, UVOIPStatics::GetNumBufferedPackets(), UVOIPStatics::GetBufferingDelay());
+			QueuedData.bIsEnvelopeBound = false;
 		}
 	}
 
-	if (QueuedData->AudioComponent != NULL)
+	if (QueuedData.VoipSynthComponent != nullptr)
 	{
-		USoundWaveProcedural* SoundStreaming = CastChecked<USoundWaveProcedural>(QueuedData->AudioComponent->Sound);
-		if (SoundStreaming->GetAvailableAudioByteCount() == 0)
+		if (!QueuedData.bIsActive)
 		{
-			UE_LOG(LogVoiceEngine, Log, TEXT("VOIP audio component was starved!"));
+			QueuedData.bIsActive = true;
+			FVoiceSettings InSettings;
+			UVOIPTalker* OwningTalker = nullptr;
+
+			OwningTalker = UVOIPStatics::GetVOIPTalkerForPlayer(RemoteTalkerId, InSettings);
+
+			ApplyVoiceSettings(QueuedData.VoipSynthComponent, InSettings);
+
+			QueuedData.VoipSynthComponent->ResetBuffer(InSampleCount, UVOIPStatics::GetBufferingDelay());
+			QueuedData.VoipSynthComponent->Start();
+			QueuedData.CachedTalkerPtr = OwningTalker;
+
+			if (OwningTalker)
+			{
+				if (!QueuedData.bIsEnvelopeBound)
+				{
+					QueuedData.VoipSynthComponent->OnAudioEnvelopeValueNative.AddUObject(OwningTalker, &UVOIPTalker::OnAudioComponentEnvelopeValue);
+					QueuedData.bIsEnvelopeBound = true;
+				}
+
+				OwningTalker->OnTalkingBegin(QueuedData.VoipSynthComponent->GetAudioComponent());
+			}
 		}
-		SoundStreaming->QueueAudio(DecompressedVoiceBuffer.GetData(), BytesWritten);
+
+		QueuedData.VoipSynthComponent->SubmitPacket((float*)DecompressedVoiceBuffer.GetData(), BytesWritten, InSampleCount, EVoipStreamDataFormat::Int16);
 	}
 
 	return S_OK;
@@ -351,64 +519,180 @@ void FVoiceEngineSteam::TickTalkers(float DeltaTime)
 	{
 		FRemoteTalkerDataSteam& RemoteData = It.Value();
 		double TimeSince = CurTime - RemoteData.LastSeen;
-		if (TimeSince >= 5.0)
+
+		if (RemoteData.VoipSynthComponent->IsIdling() && RemoteData.bIsActive)
 		{
-			// Dump the whole talker
-			if (RemoteData.AudioComponent)
+			RemoteData.VoipSynthComponent->Stop();
+
+			UAudioComponent* AudioComponent = RemoteData.VoipSynthComponent->GetAudioComponent();
+			if (AudioComponent->IsRegistered())
+			{
+				AudioComponent->UnregisterComponent();
+			}
+
+			//If the UVOIPTalker associated with this is still alive, notify it that this player is done talking.
+			if (UVOIPStatics::IsVOIPTalkerStillAlive(RemoteData.CachedTalkerPtr))
 			{
-				RemoteData.AudioComponent->Stop();
-				RemoteData.AudioComponent = NULL;
+				RemoteData.CachedTalkerPtr->OnTalkingEnd();
 			}
+			RemoteData.bIsActive = false;
 
-			It.RemoveCurrent();
+			RemoteData.Reset();
+		}
+		else if (TimeSince >= UVOIPStatics::GetRemoteTalkerTimeoutDuration())
+		{
+			// Dump the whole talker
+			RemoteData.Reset();
 		}
 	}
 }
 
-void FVoiceEngineSteam::Tick(float DeltaTime) 
+void FVoiceEngineSteam::Tick(float DeltaTime)
 {
-	// Check available voice one a frame, this value changes after calling GetVoice()
-	AvailableVoiceResult = SteamUserPtr->GetAvailableVoice(&CompressedBytesAvailable, NULL, 0);
-	//UE_LOG(LogVoiceEngine, Log, TEXT("TickResult: %s"), *SteamVoiceResult(AvailableVoiceResult));
+	// Check available voice once a frame, this value changes after calling GetVoiceData()
+	AvailableVoiceResult = VoiceCapture->GetCaptureState(UncompressedBytesAvailable);
+
 	TickTalkers(DeltaTime);
 }
 
-void FVoiceEngineSteam::OnAudioFinished(UAudioComponent* AC)
+void FVoiceEngineSteam::GenerateVoiceData(USoundWaveProcedural* InProceduralWave, int32 SamplesRequired, const FUniqueNetId& TalkerId)
 {
-	for (FRemoteTalkerData::TIterator It(RemoteTalkerBuffers); It; ++It)
+	FRemoteTalkerDataSteam* QueuedData = RemoteTalkerBuffers.Find(FUniqueNetIdWrapper(TalkerId.AsShared()));
+	if (QueuedData)
 	{
-		FRemoteTalkerDataSteam& RemoteData = It.Value();
-		if (RemoteData.AudioComponent->IsPendingKill() || AC == RemoteData.AudioComponent)
+		const int32 SampleSize = sizeof(uint16) * DEFAULT_NUM_VOICE_CHANNELS;
+
 		{
-			UE_LOG(LogVoiceEngine, Log, TEXT("Removing VOIP AudioComponent for SteamId: %s"), *It.Key().ToDebugString());
-			RemoteData.AudioComponent = NULL;
-			break;
+			FScopeLock ScopeLock(&QueuedData->QueueLock);
+			QueuedData->CurrentUncompressedDataQueueSize = QueuedData->UncompressedDataQueue.Num();
+			const int32 AvailableSamples = QueuedData->CurrentUncompressedDataQueueSize / SampleSize;
+			if (AvailableSamples >= SamplesRequired)
+			{
+				UE_LOG(LogVoiceEngine, Verbose, TEXT("GenerateVoiceData %d / %d"), AvailableSamples, SamplesRequired);
+				const int32 SamplesBytesTaken = AvailableSamples * SampleSize;
+				InProceduralWave->QueueAudio(QueuedData->UncompressedDataQueue.GetData(), SamplesBytesTaken);
+				QueuedData->UncompressedDataQueue.RemoveAt(0, SamplesBytesTaken, false);
+				QueuedData->CurrentUncompressedDataQueueSize -= (SamplesBytesTaken);
+			}
+			else
+			{
+				UE_LOG(LogVoiceEngine, Verbose, TEXT("Voice underflow"));
+			}
 		}
 	}
-	UE_LOG(LogVoiceEngine, Verbose, TEXT("Audio Finished"));
 }
 
-void FVoiceEngineSteam::OnPostLoadMap(UWorld*)
+void FVoiceEngineSteam::OnAudioFinished()
 {
 	for (FRemoteTalkerData::TIterator It(RemoteTalkerBuffers); It; ++It)
 	{
 		FRemoteTalkerDataSteam& RemoteData = It.Value();
-		if (RemoteData.AudioComponent)
+		if (RemoteData.VoipSynthComponent->IsIdling())
 		{
-			RemoteData.AudioComponent->Play();
+			UE_LOG(LogVoiceEngine, Log, TEXT("Removing VOIP AudioComponent for Id: %s"), *It.Key().ToDebugString());
+			RemoteData.VoipSynthComponent->Stop();
+			RemoteData.bIsActive = false;
+			break;
 		}
 	}
+	UE_LOG(LogVoiceEngine, Verbose, TEXT("Audio Finished"));
 }
 
-FString FVoiceEngineSteam::GetVoiceDebugState() const 
+FString FVoiceEngineSteam::GetVoiceDebugState() const
 {
 	FString Output;
-	Output = FString::Printf(TEXT("IsRecording: %d\n DataReady: 0x%08x State:%s\n BufferRemaining: %d\n"),
-		IsRecording(), 
+	Output = FString::Printf(TEXT("IsRecording: %d\n DataReady: 0x%08x State:%s\n UncompressedBytes: %d\n CompressedBytes: %d\n"),
+		IsRecording(),
 		GetVoiceDataReadyFlags(),
-		*SteamVoiceResult(AvailableVoiceResult),
-		CompressedVoiceBuffer.Num()
-		);
+		EVoiceCaptureState::ToString(AvailableVoiceResult),
+		UncompressedBytesAvailable,
+		CompressedBytesAvailable
+	);
+
+	// Add remainder size
+	for (int32 Idx = 0; Idx < MAX_SPLITSCREEN_TALKERS; Idx++)
+	{
+		Output += FString::Printf(TEXT("Remainder[%d] %d\n"), Idx, PlayerVoiceData[Idx].VoiceRemainderSize);
+	}
 
 	return Output;
 }
+
+bool FVoiceEngineSteam::Exec(UWorld* InWorld, const TCHAR* Cmd, FOutputDevice& Ar)
+{
+	bool bWasHandled = false;
+
+	if (FParse::Command(&Cmd, TEXT("vcvbr")))
+	{
+		// vcvbr <true/false>
+		FString VBRStr = FParse::Token(Cmd, false);
+		int32 ShouldVBR = FPlatformString::Atoi(*VBRStr);
+		bool bVBR = ShouldVBR != 0;
+		if (VoiceEncoder.IsValid())
+		{
+			if (!VoiceEncoder->SetVBR(bVBR))
+			{
+				UE_LOG(LogVoice, Warning, TEXT("Failed to set VBR %d"), bVBR);
+			}
+		}
+
+		bWasHandled = true;
+	}
+	else if (FParse::Command(&Cmd, TEXT("vcbitrate")))
+	{
+		// vcbitrate <bitrate>
+		FString BitrateStr = FParse::Token(Cmd, false);
+		int32 NewBitrate = !BitrateStr.IsEmpty() ? FPlatformString::Atoi(*BitrateStr) : 0;
+		if (VoiceEncoder.IsValid() && NewBitrate > 0)
+		{
+			if (!VoiceEncoder->SetBitrate(NewBitrate))
+			{
+				UE_LOG(LogVoice, Warning, TEXT("Failed to set bitrate %d"), NewBitrate);
+			}
+		}
+
+		bWasHandled = true;
+	}
+	else if (FParse::Command(&Cmd, TEXT("vccomplexity")))
+	{
+		// vccomplexity <complexity>
+		FString ComplexityStr = FParse::Token(Cmd, false);
+		int32 NewComplexity = !ComplexityStr.IsEmpty() ? FPlatformString::Atoi(*ComplexityStr) : -1;
+		if (VoiceEncoder.IsValid() && NewComplexity >= 0)
+		{
+			if (!VoiceEncoder->SetComplexity(NewComplexity))
+			{
+				UE_LOG(LogVoice, Warning, TEXT("Failed to set complexity %d"), NewComplexity);
+			}
+		}
+
+		bWasHandled = true;
+	}
+	else if (FParse::Command(&Cmd, TEXT("vcdump")))
+	{
+		if (VoiceCapture.IsValid())
+		{
+			VoiceCapture->DumpState();
+		}
+
+		if (VoiceEncoder.IsValid())
+		{
+			VoiceEncoder->DumpState();
+		}
+
+		for (FRemoteTalkerData::TIterator It(RemoteTalkerBuffers); It; ++It)
+		{
+			FRemoteTalkerDataSteam& RemoteData = It.Value();
+			if (RemoteData.VoiceDecoder.IsValid())
+			{
+				RemoteData.VoiceDecoder->DumpState();
+			}
+		}
+
+		bWasHandled = true;
+	}
+
+	return bWasHandled;
+}
+
+
diff --git a/Engine/Plugins/Online/OnlineSubsystemSteam/Source/Private/VoiceEngineSteam.h b/Engine/Plugins/Online/OnlineSubsystemSteam/Source/Private/VoiceEngineSteam.h
--- a/Engine/Plugins/Online/OnlineSubsystemSteam/Source/Private/VoiceEngineSteam.h
+++ b/Engine/Plugins/Online/OnlineSubsystemSteam/Source/Private/VoiceEngineSteam.h
@@ -3,123 +3,182 @@
 #pragma once
 
 #include "CoreMinimal.h"
-#include "Interfaces/VoiceInterface.h"
-#include "Net/VoiceDataCommon.h"
+#include "UObject/GCObject.h"
 #include "OnlineSubsystemTypes.h"
-#include "OnlineSubsystemSteamPrivate.h"
 #include "OnlineSubsystemSteamTypes.h"
-#include "UObject/GCObject.h"
-#include "OnlineSubsystemSteamPackage.h"
+#include "Interfaces/VoiceInterface.h"
+#include "Net/VoiceDataCommon.h"
+#include "Interfaces/VoiceCapture.h"
+#include "Interfaces/VoiceCodec.h"
+#include "OnlineSubsystemUtilsPackage.h"
+#include "VoipListenerSynthComponent.h"
 
-class FOnlineSubsystemSteam;
-class UAudioComponent;
+class IOnlineSubsystem;
+class FUniqueNetIdSteam;
+class IVoiceDecoder;
+class IVoiceEncoder;
+class IVoiceCapture;
 
-/** 
- * Remote voice data playing on a single client
- */
-class FRemoteTalkerDataSteam
+#define INVALID_INDEX -1
+
+/**
+* Container for unprocessed voice data
+*/
+struct FLocalVoiceDataSteam
 {
-public:
-	FRemoteTalkerDataSteam() :
-		LastSeen(0.0),
-		AudioComponent(NULL)
+	FLocalVoiceDataSteam() :
+		VoiceRemainderSize(0)
 	{
 	}
 
-	// Receive side timestamp since last voice packet fragment
+	/** Amount of voice data not encoded last time */
+	uint32 VoiceRemainderSize;
+	/** Voice sample data not encoded last time */
+	TArray<uint8> VoiceRemainder;
+};
+
+/**
+* Remote voice data playing on a single client
+*/
+class FRemoteTalkerDataSteam
+{
+public:
+
+	FRemoteTalkerDataSteam();
+	/** Required for TMap FindOrAdd() */
+	FRemoteTalkerDataSteam(FRemoteTalkerDataSteam&& Other);
+	~FRemoteTalkerDataSteam();
+
+	/** Reset the talker after long periods of silence */
+	void Reset();
+	/** Cleanup the talker before unregistration */
+	void Cleanup();
+
+	/** Maximum size of a single decoded packet */
+	int32 MaxUncompressedDataSize;
+	/** Maximum size of the outgoing playback queue */
+	int32 MaxUncompressedDataQueueSize;
+	/** Amount of data currently in the outgoing playback queue */
+	int32 CurrentUncompressedDataQueueSize;
+
+	/** Receive side timestamp since last voice packet fragment */
 	double LastSeen;
-	// Audio component playing this buffer (only valid on remote instances)
-	class UAudioComponent* AudioComponent;
+	/** Number of frames starved of audio */
+	int32 NumFramesStarved;
+	/** Synth component playing this buffer (only valid on remote instances) */
+	UVoipListenerSynthComponent* VoipSynthComponent;
+	/** Cached Talker Ptr. Is checked against map before use to ensure it has not been destroyed. */
+	UVOIPTalker* CachedTalkerPtr;
+	/** Boolean used to ensure that we only bind the VOIP talker to the SynthComponent's corresponding envelope delegate once. */
+	bool bIsEnvelopeBound;
+	/** Boolean flag used to tell whether this synth component is currently consuming incoming voice packets. */
+	bool bIsActive;
+	/** Buffer for outgoing audio intended for procedural streaming */
+	mutable FCriticalSection QueueLock;
+	TArray<uint8> UncompressedDataQueue;
+	/** Per remote talker voice decoding state */
+	TSharedPtr<IVoiceDecoder> VoiceDecoder;
 };
 
 /**
- * The Steam implementation of the voice engine 
- */
-class FVoiceEngineSteam : public IVoiceEngine
+* Generic implementation of voice engine, using Voice module for capture/codec
+*/
+class FVoiceEngineSteam : public IVoiceEngine, public FSelfRegisteringExec
 {
- 	class FVoiceSerializeHelper : public FGCObject
- 	{
- 		/** Reference to audio components */
- 		FVoiceEngineSteam* VoiceEngine;
- 		FVoiceSerializeHelper() :
- 			VoiceEngine(NULL)
- 		{}
- 
- 	public:
- 
- 		FVoiceSerializeHelper(FVoiceEngineSteam* InVoiceEngine) :
- 			VoiceEngine(InVoiceEngine)
- 		{}
- 		~FVoiceSerializeHelper() {}
- 		
- 		/** FGCObject interface */
- 		virtual void AddReferencedObjects(FReferenceCollector& Collector) override
- 		{
- 			// Prevent garbage collection of audio components
- 			for (FRemoteTalkerData::TIterator It(VoiceEngine->RemoteTalkerBuffers); It; ++It)
- 			{
- 				FRemoteTalkerDataSteam& RemoteData = It.Value();
- 				if (RemoteData.AudioComponent)
- 				{
- 					Collector.AddReferencedObject(RemoteData.AudioComponent);
- 				}
- 			}
- 		}
- 	};
+	class FVoiceSerializeHelper : public FGCObject
+	{
+		/** Reference to audio components */
+		FVoiceEngineSteam* VoiceEngine;
+		FVoiceSerializeHelper() :
+			VoiceEngine(nullptr)
+		{}
+
+	public:
+
+		FVoiceSerializeHelper(FVoiceEngineSteam* InVoiceEngine) :
+			VoiceEngine(InVoiceEngine)
+		{}
+		~FVoiceSerializeHelper() {}
+
+		/** FGCObject interface */
+		virtual void AddReferencedObjects(FReferenceCollector& Collector) override
+		{
+			// Prevent garbage collection of audio components
+			for (FRemoteTalkerData::TIterator It(VoiceEngine->RemoteTalkerBuffers); It; ++It)
+			{
+				FRemoteTalkerDataSteam& RemoteData = It.Value();
+				if (RemoteData.VoipSynthComponent)
+				{
+					Collector.AddReferencedObject(RemoteData.VoipSynthComponent);
+				}
+			}
+		}
+	};
 
 	friend class FVoiceSerializeHelper;
 
 	/** Mapping of UniqueIds to the incoming voice data and their audio component */
-	typedef TMap<FUniqueNetIdSteam, FRemoteTalkerDataSteam> FRemoteTalkerData;
+	// TODO: find a better way to make this ID class work so it stops screwing everything up.
+	typedef TMap<FUniqueNetIdWrapper, FRemoteTalkerDataSteam> FRemoteTalkerData;
 
-	/** Reference to the main Steam subsystem */
-	class FOnlineSubsystemSteam* SteamSubsystem;
-    /** Steam User interface */
+	/** Reference to the main online subsystem */
+	IOnlineSubsystem* OnlineSubsystem;
+
+	/** Steam User interface */
 	class ISteamUser* SteamUserPtr;
-    /** Steam Friends interface */
+	/** Steam Friends interface */
 	class ISteamFriends* SteamFriendsPtr;
-    /** User index currently holding onto the voice interface */
+
+	FLocalVoiceDataSteam PlayerVoiceData[MAX_SPLITSCREEN_TALKERS];
+	/** Reference to voice capture device */
+	TSharedPtr<IVoiceCapture> VoiceCapture;
+	/** Reference to voice encoding object */
+	TSharedPtr<IVoiceEncoder> VoiceEncoder;
+
+	/** User index currently holding onto the voice interface */
 	int32 OwningUserIndex;
+	/** Amount of uncompressed data available this frame */
+	uint32 UncompressedBytesAvailable;
 	/** Amount of compressed data available this frame */
 	uint32 CompressedBytesAvailable;
-	/** Result of call to GetAvailableVoice() this frame*/
-	EVoiceResult AvailableVoiceResult;
-    /** Have we stopped Steam voice but are waiting for its completion */
+	/** Current frame state of voice capture */
+	EVoiceCaptureState::Type AvailableVoiceResult;
+	/** Have we stopped capturing voice but are waiting for its completion */
 	mutable bool bPendingFinalCapture;
-    /** State of voice recording */
+	/** State of voice recording */
 	bool bIsCapturing;
 
-	/** Data from Steamworks, waiting to send to network. */
+	/** Data from voice codec, waiting to send to network. */
 	TArray<uint8> CompressedVoiceBuffer;
 	/** Data from network playing on an audio component. */
 	FRemoteTalkerData RemoteTalkerBuffers;
-	/** Voice decompression buffer, shared by all talkers */
+	/** Voice decompression buffer, shared by all talkers, valid during SubmitRemoteVoiceData */
 	TArray<uint8> DecompressedVoiceBuffer;
 	/** Serialization helper */
-	class FVoiceSerializeHelper* SerializeHelper;
+	FVoiceSerializeHelper* SerializeHelper;
 
 	/**
-	 * Determines if the specified index is the owner or not
-	 *
-	 * @param InIndex the index being tested
-	 *
-	 * @return true if this is the owner, false otherwise
-	 */
+	* Determines if the specified index is the owner or not
+	*
+	* @param InIndex the index being tested
+	*
+	* @return true if this is the owner, false otherwise
+	*/
 	FORCEINLINE bool IsOwningUser(uint32 UserIndex)
 	{
 		return UserIndex >= 0 && UserIndex < MAX_SPLITSCREEN_TALKERS && OwningUserIndex == UserIndex;
 	}
 
-	/** 
-	 * Update the internal state of the voice capturing state 
-	 * Handles Steam's continual recording for "last half second" after requested stop
-	 */
+	/**
+	* Update the internal state of the voice capturing state
+	* Handles possible continuation waiting for capture stop event
+	*/
 	void VoiceCaptureUpdate() const;
 
-	/** Tell Steam to start capturing voice data */
+	/** Start capturing voice data */
 	void StartRecording() const;
 
-	/** Tell Steam to stop capturing voice data */
+	/** Stop capturing voice data */
 	void StopRecording() const;
 
 	/** Called when "last half second" is over */
@@ -128,19 +187,31 @@ class FVoiceEngineSteam : public IVoiceEngine
 	/** @return is active recording occurring at the moment */
 	bool IsRecording() const { return bIsCapturing || bPendingFinalCapture; }
 
+	/**
+	* Callback from streaming audio when data is requested for playback
+	*
+	* @param InProceduralWave SoundWave requesting more data
+	* @param SamplesRequired number of samples needed for immediate playback
+	* @param TalkerId id of the remote talker to allocate voice data for
+	*/
+	void GenerateVoiceData(USoundWaveProcedural* InProceduralWave, int32 SamplesRequired, const FUniqueNetId& TalkerId);
+
 PACKAGE_SCOPE:
 
 	/** Constructor */
 	FVoiceEngineSteam() :
-		SteamSubsystem(NULL),
-		SteamUserPtr(NULL),
-		SteamFriendsPtr(NULL),
+		OnlineSubsystem(nullptr),
+		SteamUserPtr(nullptr),
+		SteamFriendsPtr(nullptr),
+		VoiceCapture(nullptr),
+		VoiceEncoder(nullptr),
 		OwningUserIndex(INVALID_INDEX),
+		UncompressedBytesAvailable(0),
 		CompressedBytesAvailable(0),
-		AvailableVoiceResult(k_EVoiceResultNotInitialized),
+		AvailableVoiceResult(EVoiceCaptureState::UnInitialized),
 		bPendingFinalCapture(false),
 		bIsCapturing(false),
-		SerializeHelper(NULL)
+		SerializeHelper(nullptr)
 	{};
 
 	// IVoiceEngine
@@ -148,24 +219,26 @@ PACKAGE_SCOPE:
 
 public:
 
-	FVoiceEngineSteam(FOnlineSubsystemSteam* InSteamSubsystem);
+	FVoiceEngineSteam(IOnlineSubsystem* InSubsystem);
 	virtual ~FVoiceEngineSteam();
 
 	// IVoiceEngine
-    virtual uint32 StartLocalVoiceProcessing(uint32 LocalUserNum) override;
-    virtual uint32 StopLocalVoiceProcessing(uint32 LocalUserNum) override;
-    virtual uint32 StartRemoteVoiceProcessing(const FUniqueNetId& UniqueId) override
+	virtual uint32 StartLocalVoiceProcessing(uint32 LocalUserNum) override;
+	virtual uint32 StopLocalVoiceProcessing(uint32 LocalUserNum) override;
+
+	virtual uint32 StartRemoteVoiceProcessing(const FUniqueNetId& UniqueId) override
 	{
-		// Not needed in Steam
+		// Not needed
 		return S_OK;
 	}
-    virtual uint32 StopRemoteVoiceProcessing(const FUniqueNetId& UniqueId) override
+
+	virtual uint32 StopRemoteVoiceProcessing(const FUniqueNetId& UniqueId) override
 	{
-		// Not needed in Steam
+		// Not needed
 		return S_OK;
 	}
 
-    virtual uint32 RegisterLocalTalker(uint32 LocalUserNum) override
+	virtual uint32 RegisterLocalTalker(uint32 LocalUserNum) override
 	{
 		if (OwningUserIndex == INVALID_INDEX)
 		{
@@ -176,7 +249,7 @@ public:
 		return E_FAIL;
 	}
 
-    virtual uint32 UnregisterLocalTalker(uint32 LocalUserNum) override
+	virtual uint32 UnregisterLocalTalker(uint32 LocalUserNum) override
 	{
 		if (IsOwningUser(LocalUserNum))
 		{
@@ -187,31 +260,27 @@ public:
 		return E_FAIL;
 	}
 
-    virtual uint32 RegisterRemoteTalker(const FUniqueNetId& UniqueId) override
+	virtual uint32 RegisterRemoteTalker(const FUniqueNetId& UniqueId) override
 	{
-		// Not needed in Steam
+		// Not needed
 		return S_OK;
 	}
 
-    virtual uint32 UnregisterRemoteTalker(const FUniqueNetId& UniqueId) override
-	{
-		// Not needed in Steam
-		return S_OK;
-	}
+	virtual uint32 UnregisterRemoteTalker(const FUniqueNetId& UniqueId) override;
 
-    virtual bool IsHeadsetPresent(uint32 LocalUserNum) override
+	virtual bool IsHeadsetPresent(uint32 LocalUserNum) override
 	{
 		return IsOwningUser(LocalUserNum) ? true : false;
 	}
 
-    virtual bool IsLocalPlayerTalking(uint32 LocalUserNum) override
+	virtual bool IsLocalPlayerTalking(uint32 LocalUserNum) override
 	{
 		return (GetVoiceDataReadyFlags() & (LocalUserNum << 1)) != 0;
 	}
 
 	virtual bool IsRemotePlayerTalking(const FUniqueNetId& UniqueId) override
 	{
-		return RemoteTalkerBuffers.Find((const FUniqueNetIdSteam&)UniqueId) != NULL;
+		return RemoteTalkerBuffers.Find(FUniqueNetIdWrapper(UniqueId.AsShared())) != nullptr;
 	}
 
 	virtual uint32 GetVoiceDataReadyFlags() const override;
@@ -221,25 +290,33 @@ public:
 		return S_OK;
 	}
 
-	virtual uint32 ReadLocalVoiceData(uint32 LocalUserNum, uint8* Data, uint32* Size) override;
-	virtual uint32 SubmitRemoteVoiceData(const FUniqueNetId& RemoteTalkerId, uint8* Data, uint32* Size) override;
+	virtual uint32 ReadLocalVoiceData(uint32 LocalUserNum, uint8* Data, uint32* Size) override { return ReadLocalVoiceData(LocalUserNum, Data, Size, nullptr); }
+	virtual uint32 ReadLocalVoiceData(uint32 LocalUserNum, uint8* Data, uint32* Size, uint64* OutSampleCount) override;
+
+	virtual uint32 SubmitRemoteVoiceData(const FUniqueNetId& RemoteTalkerId, uint8* Data, uint32* Size)
+	{
+		checkf(false, TEXT("Please use the following function signature instead: SubmitRemoteVoiceData(const FUniqueNetIdWrapper& RemoteTalkerId, uint8* Data, uint32* Size, uint64& InSampleCount)"));
+		return 0;
+	}
+	virtual uint32 SubmitRemoteVoiceData(const FUniqueNetIdWrapper& RemoteTalkerId, uint8* Data, uint32* Size, uint64& InSampleCount) override;
+
 	virtual void Tick(float DeltaTime) override;
 	FString GetVoiceDebugState() const override;
 
-	/**
-	 * Update the state of all remote talkers, possibly dropping data or the talker entirely
-	 */
-	void TickTalkers(float DeltaTime);
+	// FSelfRegisteringExec
+	virtual bool Exec(UWorld* InWorld, const TCHAR* Cmd, FOutputDevice& Ar) override;
+
+private:
 
 	/**
-	 * Delegate triggered when an audio component Stop() function is called
-	 */
-	void OnAudioFinished(UAudioComponent* AC);
+	* Update the state of all remote talkers, possibly dropping data or the talker entirely
+	*/
+	void TickTalkers(float DeltaTime);
 
 	/**
-	 * Delegate that fixes up remote audio components when the level changes
-	 */
-	void OnPostLoadMap(UWorld*);
+	* Delegate triggered when an audio component Stop() function is called
+	*/
+	void OnAudioFinished();
 };
 
-typedef TSharedPtr<FVoiceEngineSteam, ESPMode::ThreadSafe> FVoiceEngineSteamPtr;
+typedef TSharedPtr<FVoiceEngineSteam, ESPMode::ThreadSafe> FVoiceEngineImplPtr;
diff --git a/Engine/Plugins/Online/OnlineSubsystemSteam/Source/Private/VoiceInterfaceSteam.cpp b/Engine/Plugins/Online/OnlineSubsystemSteam/Source/Private/VoiceInterfaceSteam.cpp
--- a/Engine/Plugins/Online/OnlineSubsystemSteam/Source/Private/VoiceInterfaceSteam.cpp
+++ b/Engine/Plugins/Online/OnlineSubsystemSteam/Source/Private/VoiceInterfaceSteam.cpp
@@ -1,856 +1,2 @@
 // Copyright 1998-2018 Epic Games, Inc. All Rights Reserved.
-
 #include "VoiceInterfaceSteam.h"
-#include "Misc/ConfigCacheIni.h"
-#include "EngineGlobals.h"
-#include "GameFramework/OnlineReplStructs.h"
-#include "Engine/Engine.h"
-#include "GameFramework/PlayerController.h"
-#include "Engine/LocalPlayer.h"
-#include "OnlineSubsystemSteam.h"
-#include "VoiceEngineSteam.h"
-#include "VoiceModule.h"
-#include "OnlineIdentityInterfaceSteam.h"
-#include "OnlineSessionInterfaceSteam.h"
-
-/** Largest size to attempt to transmit */
-#define MAX_VOICE_PACKET_SIZE_STEAM 1 * 1024
-
-FOnlineVoiceSteam::FOnlineVoiceSteam(FOnlineSubsystemSteam* InSteamSubsystem) :
-	SteamSubsystem(InSteamSubsystem),
-	VoiceEngine(NULL)
-{
-	check(InSteamSubsystem);
-}
-
-FOnlineVoiceSteam::~FOnlineVoiceSteam()
-{
-	LocalTalkers.Empty();
-	RemoteTalkers.Empty();
-	VoiceEngine = NULL;
-}
-
-bool FOnlineVoiceSteam::Init()
-{
-	bool bSuccess = false;
-
-	if (!GConfig->GetInt(TEXT("OnlineSubsystem"),TEXT("MaxLocalTalkers"), MaxLocalTalkers, GEngineIni))
-	{
-		MaxLocalTalkers = MAX_SPLITSCREEN_TALKERS;
-		UE_LOG(LogVoice, Warning, TEXT("Missing MaxLocalTalkers key in OnlineSubsystem of DefaultEngine.ini"));
-	}
-	if (!GConfig->GetInt(TEXT("OnlineSubsystem"),TEXT("MaxRemoteTalkers"), MaxRemoteTalkers, GEngineIni))
-	{
-		MaxRemoteTalkers = MAX_REMOTE_TALKERS;
-		UE_LOG(LogVoice, Warning, TEXT("Missing MaxRemoteTalkers key in OnlineSubsystem of DefaultEngine.ini"));
-	}
-	if (!GConfig->GetFloat(TEXT("OnlineSubsystem"),TEXT("VoiceNotificationDelta"), VoiceNotificationDelta, GEngineIni))
-	{
-		VoiceNotificationDelta = 0.2;
-		UE_LOG(LogVoice, Warning, TEXT("Missing VoiceNotificationDelta key in OnlineSubsystem of DefaultEngine.ini"));
-	}
-
-	bool bHasVoiceEnabled = false;
-	if (GConfig->GetBool(TEXT("OnlineSubsystem"), TEXT("bHasVoiceEnabled"), bHasVoiceEnabled, GEngineIni) && bHasVoiceEnabled)
-	{
-		SessionInt = (FOnlineSessionSteam*)SteamSubsystem->GetSessionInterface().Get();
-		IdentityInt = (FOnlineIdentitySteam*)SteamSubsystem->GetIdentityInterface().Get();
-
-		bSuccess = SessionInt && IdentityInt;
-
-		if (bSuccess)
-		{
-			const bool bVoiceEngineForceDisable = SteamSubsystem->IsDedicated() || GIsBuildMachine;
-			if (!bVoiceEngineForceDisable)
-			{
-				VoiceEngine = MakeShareable(new FVoiceEngineSteam(SteamSubsystem));
-				bSuccess = VoiceEngine->Init(MaxLocalTalkers, MaxRemoteTalkers);
-			}
-			else
-			{
-				MaxLocalTalkers = 0;
-				MaxRemoteTalkers = 0;
-			}
-		}
-
-		LocalTalkers.Init(FLocalTalker(), MaxLocalTalkers);
-		RemoteTalkers.Empty(MaxRemoteTalkers);
-
-		if (!bSuccess)
-		{
-			// Not necessary to log here since VoiceEngine::Init() will report its own failure
-			//UE_LOG(LogVoice, Warning, TEXT("Failed to initialize Steam voice interface"));
-
-			LocalTalkers.Empty();
-			RemoteTalkers.Empty();
-			VoiceEngine = NULL;
-		}
-	}
-	else
-	{
-		UE_LOG(LogVoice, Log, TEXT("Voice interface disabled by config [OnlineSubsystem].bHasVoiceEnabled"));
-	}
-
-	return bSuccess;
-}
-
-void FOnlineVoiceSteam::ClearVoicePackets()
-{
-	for (uint32 Index = 0; Index < MAX_SPLITSCREEN_TALKERS; Index++)
-	{
-		FVoicePacketSteam& LocalPacket = VoiceData.LocalPackets[Index];
-		if (LocalPacket.Length > 0)
-		{
-			// Mark the local packet as processed
-			LocalPacket.Length = 0;
-		}
-	}
-}
-
-void FOnlineVoiceSteam::Tick(float DeltaTime) 
-{
-	if (!SteamSubsystem->IsDedicated())
-	{
-		SCOPE_CYCLE_COUNTER(STAT_Voice_Interface);
-
-		// If we aren't in a networked match, no need to update networked voice
-		if (SessionInt && SessionInt->GetNumSessions() > 0)
-		{
-			// Processing voice data only valid with a voice engine to capture/play
-			if (VoiceEngine.IsValid())
-			{
-				VoiceEngine->Tick(DeltaTime);
-
-				// Queue local packets for sending via the network
-				ProcessLocalVoicePackets();
-				// Submit queued packets to audio system
-				ProcessRemoteVoicePackets();
-				// Fire off any talking notifications for hud display
-				ProcessTalkingDelegates(DeltaTime);
-			}
-		}
-	}
-}
-
-void FOnlineVoiceSteam::StartNetworkedVoice(uint8 LocalUserNum)
-{
-	// Validate the range of the entry
-	if (LocalUserNum < MAX_LOCAL_PLAYERS)
-	{
-		LocalTalkers[LocalUserNum].bHasNetworkedVoice = true;
-		UE_LOG(LogVoice, Log, TEXT("Starting networked voice for user: %d"), LocalUserNum);
-	}
-	else
-	{
-		UE_LOG(LogVoice, Log, TEXT("Invalid user specified in StartNetworkedVoice(%d)"),
-			(uint32)LocalUserNum);
-	}
-}
-
-void FOnlineVoiceSteam::StopNetworkedVoice(uint8 LocalUserNum)
-{
-	// Validate the range of the entry
-	if (LocalTalkers.Num() > LocalUserNum && LocalUserNum >= 0 && LocalUserNum < MAX_LOCAL_PLAYERS)
-	{
-		LocalTalkers[LocalUserNum].bHasNetworkedVoice = false;
-		UE_LOG(LogVoice, Log, TEXT("Stopping networked voice for user: %d"), LocalUserNum);
-	}
-	else
-	{
-		UE_LOG(LogVoice, Log, TEXT("Invalid user specified in StopNetworkedVoice(%d)"),
-			(uint32)LocalUserNum);
-	}
-}
-
-bool FOnlineVoiceSteam::RegisterLocalTalker(uint32 LocalUserNum) 
-{
-	uint32 Return = E_FAIL;
-	if (LocalUserNum >= 0 && LocalUserNum < MAX_LOCAL_PLAYERS)
-	{
-		// Get at the local talker's cached data
-		FLocalTalker& Talker = LocalTalkers[LocalUserNum];
-		// Make local user capable of sending voice data
-		StartNetworkedVoice(LocalUserNum);
-		// Don't register talkers when voice is disabled
-		if (VoiceEngine.IsValid())
-		{
-			if (Talker.bIsRegistered == false)
-			{
-				// Register the talker locally
-				Return = VoiceEngine->RegisterLocalTalker(LocalUserNum);
-				UE_LOG(LogVoice, Log, TEXT("RegisterLocalTalker(%d) returned 0x%08X"),	LocalUserNum, Return);
-				if (Return == S_OK)
-				{
-					Talker.bIsRegistered = true;
-					// Kick off the processing mode
-					Return = VoiceEngine->StartLocalVoiceProcessing(LocalUserNum);
-					UE_LOG(LogVoice, Log, TEXT("StartLocalProcessing(%d) returned 0x%08X"), LocalUserNum, Return);
-				}
-			}
-			else
-			{
-				// Just say yes, we registered fine
-				Return = S_OK;
-			}
-			
-			// @todo ONLINE - update mute list?
-		}
-		else
-		{
-			// Not properly logged in, so skip voice for them
-			Talker.bIsRegistered = false;
-		}
-	}
-	else
-	{
-		UE_LOG(LogVoice, Log, TEXT("Invalid user specified in RegisterLocalTalker(%d)"), LocalUserNum);
-	}
-	return Return == S_OK;
-}
-
-void FOnlineVoiceSteam::RegisterLocalTalkers()
-{
-	UE_LOG(LogVoice, Log, TEXT("Registering all local talkers"));
-	// Loop through the 4 available players and register them
-	for (uint32 Index = 0; Index < MAX_LOCAL_PLAYERS; Index++)
-	{
-		// Register the local player as a local talker
-		RegisterLocalTalker(Index);
-	}
-}
-
-bool FOnlineVoiceSteam::UnregisterLocalTalker(uint32 LocalUserNum) 
-{
-	uint32 Return = S_OK;
-	if (LocalUserNum >= 0 && LocalUserNum < MAX_LOCAL_PLAYERS)
-	{
-		// Get at the local talker's cached data
-		FLocalTalker& Talker = LocalTalkers[LocalUserNum];
-		// Skip the unregistration if not registered
-		if (Talker.bIsRegistered == true &&
-			// Or when voice is disabled
-			VoiceEngine.IsValid())
-		{
-			if (OnPlayerTalkingStateChangedDelegates.IsBound() && (Talker.bIsTalking || Talker.bWasTalking))
-			{
-				TSharedPtr<const FUniqueNetId> UniqueId = IdentityInt->GetUniquePlayerId(LocalUserNum);
-				if (UniqueId.IsValid())
-				{
-					OnPlayerTalkingStateChangedDelegates.Broadcast(UniqueId.ToSharedRef(), false);
-				}
-				else
-				{
-					UE_LOG(LogVoice, Warning, TEXT("Invalid UserId for local player %d in UnregisterLocalTalker"), LocalUserNum);
-				}
-			}
-
-			// Remove them from engine too
-			VoiceEngine->StopLocalVoiceProcessing(LocalUserNum);
-			Return = VoiceEngine->UnregisterLocalTalker(LocalUserNum);
-			UE_LOG(LogVoice, Log, TEXT("UnregisterLocalTalker(%d) returned 0x%08X"), LocalUserNum, Return);
-			Talker.bIsTalking = false;
-			Talker.bWasTalking = false;
-			Talker.bIsRegistered = false;
-		}
-	}
-	else
-	{
-		UE_LOG(LogVoice, Log, TEXT("Invalid user specified in UnregisterLocalTalker(%d)"), LocalUserNum);
-	}
-	return Return == S_OK;
-}
-
-void FOnlineVoiceSteam::UnregisterLocalTalkers()
-{
-	UE_LOG(LogVoice, Log, TEXT("Unregistering all local talkers"));
-	// Loop through the 4 available players and unregister them
-	for (uint32 Index = 0; Index < MAX_LOCAL_PLAYERS; Index++)
-	{
-		// Unregister the local player as a local talker
-		UnregisterLocalTalker(Index);
-	}
-}
-
-bool FOnlineVoiceSteam::RegisterRemoteTalker(const FUniqueNetId& UniqueId) 
-{
-	uint32 Return = E_FAIL;
-
-	// Skip this if the session isn't active
-	if (SessionInt && SessionInt->GetNumSessions() > 0 &&
-		// Or when voice is disabled
-		VoiceEngine.IsValid())
-	{
-		// See if this talker has already been registered or not
-		FRemoteTalker* Talker = FindRemoteTalker(UniqueId);
-		if (Talker == NULL)
-		{
-			// Add a new talker to our list
-			int32 AddIndex = RemoteTalkers.AddZeroed();
-			Talker = &RemoteTalkers[AddIndex];
-			// Copy the UniqueId
-			const FUniqueNetIdSteam& UniqueIdSteam = (const FUniqueNetIdSteam&)UniqueId;
-			Talker->TalkerId = MakeShareable(new FUniqueNetIdSteam(UniqueIdSteam));
-			// Register the remote talker locally
-			Return = VoiceEngine->RegisterRemoteTalker(UniqueId);
-			UE_LOG(LogVoice, Log, TEXT("RegisterRemoteTalker(%s) returned 0x%08X"),
-				*UniqueId.ToDebugString(), Return);
-		}
-		else
-		{
-			UE_LOG(LogVoice, Verbose, TEXT("Remote talker %s is being re-registered"), *UniqueId.ToDebugString());
-			Return = S_OK;
-		}
-		
-		// Update muting all of the local talkers with this remote talker
-		ProcessMuteChangeNotification();
-		// Now start processing the remote voices
-		Return = VoiceEngine->StartRemoteVoiceProcessing(UniqueId);
-		UE_LOG(LogVoice, Log, TEXT("StartRemoteVoiceProcessing(%s) returned 0x%08X"), *UniqueId.ToDebugString(), Return);
-	}
-
-	return Return == S_OK;
-}
-
-bool FOnlineVoiceSteam::UnregisterRemoteTalker(const FUniqueNetId& UniqueId) 
-{
-	uint32 Return = E_FAIL;
-
-	// Skip this if the session isn't active
-	if (SessionInt && SessionInt->GetNumSessions() > 0 &&
-		// Or when voice is disabled
-		VoiceEngine.IsValid())
-	{
-		// Make sure the talker is valid
-		if (FindRemoteTalker(UniqueId) != NULL)
-		{
-			// Find them in the talkers array and remove them
-			for (int32 Index = 0; Index < RemoteTalkers.Num(); Index++)
-			{
-				const FRemoteTalker& Talker = RemoteTalkers[Index];
-				// Is this the remote talker?
-				if (*Talker.TalkerId == UniqueId)
-				{
-					// Going to remove the talker, so if they were talking recently make sure to indicate they've stopped
-					if (OnPlayerTalkingStateChangedDelegates.IsBound() && (Talker.bIsTalking || Talker.bWasTalking))
-					{
-						OnPlayerTalkingStateChangedDelegates.Broadcast(Talker.TalkerId.ToSharedRef(), false);
-					}
-
-					RemoteTalkers.RemoveAtSwap(Index);
-					break;
-				}
-			}
-			// Remove them from voice engine
-			Return = VoiceEngine->UnregisterRemoteTalker(UniqueId);
-			UE_LOG(LogVoice, Log, TEXT("UnregisterRemoteTalker(%s) returned 0x%08X"), *UniqueId.ToDebugString(), Return);
-		}
-		else
-		{
-			UE_LOG(LogVoice, Verbose, TEXT("Unknown remote talker (%s) specified to UnregisterRemoteTalker()"), *UniqueId.ToDebugString());
-		}
-	}
-
-	return Return == S_OK;
-}
-
-void FOnlineVoiceSteam::RemoveAllRemoteTalkers()
-{
-	UE_LOG(LogVoice, Log, TEXT("Removing all remote talkers"));
-	if (VoiceEngine.IsValid())
-	{
-		// Work backwards through array removing the talkers
-		for (int32 Index = RemoteTalkers.Num() - 1; Index >= 0; Index--)
-		{
-			const FRemoteTalker& Talker = RemoteTalkers[Index];
-
-			if (OnPlayerTalkingStateChangedDelegates.IsBound() && (Talker.bIsTalking || Talker.bWasTalking))
-			{
-				OnPlayerTalkingStateChangedDelegates.Broadcast(Talker.TalkerId.ToSharedRef(), false);
-			}
-
-			uint32 Return = VoiceEngine->UnregisterRemoteTalker(*Talker.TalkerId);
-			UE_LOG(LogVoice, Log, TEXT("UnregisterRemoteTalker(%s) returned 0x%08X"), *Talker.TalkerId->ToDebugString(), Return);
-		}
-	}
-
-	// Empty the array now that they are all unregistered
-	RemoteTalkers.Empty(MaxRemoteTalkers);
-}
-
-FRemoteTalker* FOnlineVoiceSteam::FindRemoteTalker(const FUniqueNetId& UniqueId)
-{
-	for (int32 Index = 0; Index < RemoteTalkers.Num(); Index++)
-	{
-		FRemoteTalker& Talker = RemoteTalkers[Index];
-		// Compare net ids to see if they match
-		if (*Talker.TalkerId == UniqueId)
-		{
-			return &RemoteTalkers[Index];
-		}
-	}
-	return NULL;
-}
-
-bool FOnlineVoiceSteam::IsHeadsetPresent(uint32 LocalUserNum) 
-{
-	return VoiceEngine.IsValid() && VoiceEngine->IsHeadsetPresent(LocalUserNum);
-}
-
-bool FOnlineVoiceSteam::IsLocalPlayerTalking(uint32 LocalUserNum) 
-{
-	return VoiceEngine.IsValid() && VoiceEngine->IsLocalPlayerTalking(LocalUserNum);
-}
-
-bool FOnlineVoiceSteam::IsRemotePlayerTalking(const FUniqueNetId& UniqueId) 
-{
-	return VoiceEngine.IsValid() && VoiceEngine->IsRemotePlayerTalking(UniqueId);
-}
-
-bool FOnlineVoiceSteam::IsMuted(uint32 LocalUserNum, const FUniqueNetId& UniqueId) const
-{
-	if (LocalUserNum >= 0 && LocalUserNum < MAX_LOCAL_PLAYERS)
-	{
-		return IsLocallyMuted(UniqueId);
-	}
-
-	return false;
-}
-
-bool FOnlineVoiceSteam::IsLocallyMuted(const FUniqueNetId& UniqueId) const
-{
-	int32 Index = MuteList.Find((const FUniqueNetIdSteam&)UniqueId);
-	return Index != INDEX_NONE;
-}
-
-bool FOnlineVoiceSteam::IsSystemWideMuted(const FUniqueNetId& UniqueId) const
-{
-	int32 Index = SystemMuteList.Find((const FUniqueNetIdSteam&)UniqueId);
-	return Index != INDEX_NONE;
-}
-
-bool FOnlineVoiceSteam::MuteRemoteTalker(uint8 LocalUserNum, const FUniqueNetId& PlayerId, bool bIsSystemWide)
-{
-	uint32 Return = E_FAIL;
-	if (LocalUserNum >= 0 && LocalUserNum < MAX_LOCAL_PLAYERS)
-	{
-		if (bIsSystemWide)
-		{
-			// Add them to the system wide mute list
-			SystemMuteList.AddUnique((const FUniqueNetIdSteam&)PlayerId);
-			// Should update MuteList after going up to the server and coming back down
-			ProcessMuteChangeNotification();
-			Return = S_OK;
-		}
-		else
-		{
-			// Skip this if the session isn't active
-			if (SessionInt && SessionInt->GetNumSessions() > 0 &&
-				// Or if voice is disabled
-					VoiceEngine.IsValid())
-			{
-				// Find the specified talker
-				FRemoteTalker* Talker = FindRemoteTalker(PlayerId);
-				if (Talker != NULL)
-				{
-					MuteList.AddUnique((const FUniqueNetIdSteam&)PlayerId);
-					Return = S_OK;
-					UE_LOG(LogVoice, Log, TEXT("Muting remote talker (%s)"), *PlayerId.ToDebugString());
-				}
-				else
-				{
-					UE_LOG(LogVoice, Verbose, TEXT("Unknown remote talker (%s) specified to MuteRemoteTalker()"), *PlayerId.ToDebugString());
-				}
-			}
-		}
-	}
-	else
-	{
-		UE_LOG(LogVoice, Warning, TEXT("Invalid user specified in MuteRemoteTalker(%d)"), LocalUserNum);
-	}
-
-	return Return == S_OK;
-}
-
-bool FOnlineVoiceSteam::UnmuteRemoteTalker(uint8 LocalUserNum, const FUniqueNetId& PlayerId, bool bIsSystemWide)
-{
-	uint32 Return = E_FAIL;
-	if (LocalUserNum >= 0 && LocalUserNum < MAX_LOCAL_PLAYERS)
-	{
-		if (bIsSystemWide)
-		{
-			// Remove them from the mute list
-			SystemMuteList.RemoveSingleSwap((const FUniqueNetIdSteam&)PlayerId);
-			// Should update MuteList after going up to the server and coming back down
-			ProcessMuteChangeNotification();
-			Return = S_OK;
-		}
-		else
-		{
-			// Skip this if the session isn't active
-			if (SessionInt && SessionInt->GetNumSessions() > 0 &&
-				// Or if voice is disabled
-				VoiceEngine.IsValid())
-			{
-				// Find the specified talker
-				FRemoteTalker* Talker = FindRemoteTalker(PlayerId);
-				if (Talker != NULL)
-				{
-					// Make sure there isn't a system mute
-					bool bIsSystemMuted = IsSystemWideMuted(PlayerId);
-					if (!bIsSystemMuted)
-					{
-						// Remove them from the mute list
-						MuteList.RemoveSingleSwap((const FUniqueNetIdSteam&)PlayerId);
-						UE_LOG(LogVoice, Log, TEXT("Unmuting remote talker (%s)"), *PlayerId.ToDebugString());
-					}
-					Return = S_OK;
-				}
-				else
-				{
-					UE_LOG(LogVoice, Verbose, TEXT("Unknown remote talker (%s) specified to UnmuteRemoteTalker()"), *PlayerId.ToDebugString());
-				}
-			}
-		}
-	}
-	else
-	{
-		UE_LOG(LogVoice, Warning, TEXT("Invalid user specified in UnmuteRemoteTalker(%d)"), LocalUserNum);
-	}
-
-	return Return == S_OK;
-}
-
-void FOnlineVoiceSteam::ProcessMuteChangeNotification()
-{
-	// Nothing to update if there isn't an active session
-	if (VoiceEngine.IsValid())
-	{
-		if (SessionInt && SessionInt->GetNumSessions() > 0)
-		{
-			// For each local user with voice
-			for (int32 Index = 0; Index < MaxLocalTalkers; Index++)
-			{
-				// Find the very first ULocalPlayer for this ControllerId. 
-				// This is imperfect and means we cannot support voice chat properly for
-				// multiple UWorlds (but thats ok for the time being).
-				ULocalPlayer* LP = GEngine->FindFirstLocalPlayerFromControllerId(Index);
-				if (LP && LP->PlayerController)
-				{
-					// If there is a player controller, we can mute/unmute people
-					if (LocalTalkers[Index].bIsRegistered)
-					{
-						// Use the common method of checking muting
-						UpdateMuteListForLocalTalker(Index, LP->PlayerController);
-					}
-				}
-			}
-		}
-	}
-}
-
-void FOnlineVoiceSteam::UpdateMuteListForLocalTalker(int32 TalkerIndex, APlayerController* PlayerController)
-{
-	// For each registered remote talker
-	for (int32 RemoteIndex = 0; RemoteIndex < RemoteTalkers.Num(); RemoteIndex++)
-	{
-		const FRemoteTalker& Talker = RemoteTalkers[RemoteIndex];
-
-		FUniqueNetIdRepl UniqueIdRepl(Talker.TalkerId);
-
-		// Is the remote talker on this local player's mute list?
-		if (SystemMuteList.Find((FUniqueNetIdSteam&)*Talker.TalkerId) == INDEX_NONE)
-		{
-			// Unmute on the server
-			PlayerController->ServerUnmutePlayer(UniqueIdRepl);
-		}
-		else
-		{
-			// Mute on the server
-			PlayerController->ServerMutePlayer(UniqueIdRepl);
-		}
-
-		// The ServerUn/MutePlayer() functions will perform the muting based
-		// upon gameplay settings and other player's mute list
-	}
-}
-
-TSharedPtr<FVoicePacket> FOnlineVoiceSteam::SerializeRemotePacket(FArchive& Ar)
-{
-	TSharedPtr<FVoicePacketSteam> NewPacket = MakeShareable(new FVoicePacketSteam());
-	NewPacket->Serialize(Ar);
-	if (Ar.IsError() == false && NewPacket->GetBufferSize() > 0)
-	{
-		if (!SteamSubsystem->IsDedicated())
-		{
-			FUniqueNetIdMatcher PlayerMatch(*NewPacket->GetSender());
-			if (MuteList.IndexOfByPredicate(PlayerMatch) == INDEX_NONE)
-			{
-				VoiceData.RemotePackets.Add(NewPacket);
-			}
-		}
-
-		return NewPacket;
-	}
-
-	return NULL;
-}
-
-TSharedPtr<FVoicePacket> FOnlineVoiceSteam::GetLocalPacket(uint32 LocalUserNum)
-{
-	// duplicate the local copy of the data and set it on a shared pointer for destruction elsewhere
-	if (LocalUserNum >= 0 && LocalUserNum < MAX_SPLITSCREEN_TALKERS)
-	{
-		FVoicePacketSteam& VoicePacket = VoiceData.LocalPackets[LocalUserNum];
-		if (VoicePacket.GetBufferSize() > 0)
-		{
-			return MakeShareable(new FVoicePacketSteam(VoicePacket));
-		}
-	}
-
-	return NULL;
-}
-
-void FOnlineVoiceSteam::ProcessTalkingDelegates(float DeltaTime)
-{
-	// Fire off any talker notification delegates for local talkers
-	for (int32 LocalUserNum = 0; LocalUserNum < LocalTalkers.Num(); LocalUserNum++)
-	{
-		FLocalTalker& Talker = LocalTalkers[LocalUserNum];
-
-		// Only check players with voice
-		if (Talker.bIsRegistered)
-		{
-			// If the talker was not previously talking, but now is trigger the event
-			bool bShouldNotify = !Talker.bWasTalking && Talker.bIsTalking;
-			// If the talker was previously talking, but now isn't time delay the event
-			if (!bShouldNotify && Talker.bWasTalking)
-			{
-				Talker.LastNotificationTime -= DeltaTime;
-				if (Talker.LastNotificationTime <= 0.f)
-				{
-					// Clear the flag so it only activates when needed
-					Talker.bIsTalking = false;
-					Talker.LastNotificationTime = VoiceNotificationDelta;
-					bShouldNotify = true;
-				}
-			}
-
-			if (bShouldNotify)
-			{
-				// Skip all delegate handling if none are registered
-				if (OnPlayerTalkingStateChangedDelegates.IsBound())
-				{
-					TSharedPtr<const FUniqueNetId> UniqueId = IdentityInt->GetUniquePlayerId(LocalUserNum);
-					OnPlayerTalkingStateChangedDelegates.Broadcast(UniqueId.ToSharedRef(), Talker.bIsTalking);
-				}
-
-				Talker.bWasTalking = Talker.bIsTalking;
-				UE_LOG(LogVoice, Log, TEXT("Trigger Local %d %sTALKING"), LocalUserNum, Talker.bIsTalking ? TEXT("") : TEXT("NOT"));
-			}
-		}
-	}
-	// Now check all remote talkers
-	for (int32 Index = 0; Index < RemoteTalkers.Num(); Index++)
-	{
-		FRemoteTalker& Talker = RemoteTalkers[Index];
-
-		// If the talker was not previously talking, but now is trigger the event
-		bool bShouldNotify = !Talker.bWasTalking && Talker.bIsTalking;
-		// If the talker was previously talking, but now isn't time delay the event
-		if (!bShouldNotify && Talker.bWasTalking && !Talker.bIsTalking)
-		{
-			Talker.LastNotificationTime -= DeltaTime;
-			if (Talker.LastNotificationTime <= 0.f)
-			{
-				bShouldNotify = true;
-			}
-		}
-
-		if (bShouldNotify)
-		{
-			// Skip all delegate handling if none are registered
-			if (OnPlayerTalkingStateChangedDelegates.IsBound())
-			{
-				OnPlayerTalkingStateChangedDelegates.Broadcast(Talker.TalkerId.ToSharedRef(), Talker.bIsTalking);
-			}
-
-			UE_LOG(LogVoice, Log, TEXT("Trigger Remote %s %sTALKING"), *Talker.TalkerId->ToString(), Talker.bIsTalking ? TEXT("") : TEXT("NOT"));
-
-			// Clear the flag so it only activates when needed
-			Talker.bWasTalking = Talker.bIsTalking;
-			Talker.LastNotificationTime = VoiceNotificationDelta;
-		}
-	}
-}
-
-void FOnlineVoiceSteam::ProcessLocalVoicePackets()
-{
-	if (VoiceEngine.IsValid())
-	{
-		// Read the data from any local talkers
-		uint32 DataReadyFlags = VoiceEngine->GetVoiceDataReadyFlags();
-		// Skip processing if there is no data from a local talker
-		if (DataReadyFlags)
-		{
-			// Process each talker with a bit set
-			for (uint32 Index = 0; DataReadyFlags; Index++, DataReadyFlags >>= 1)
-			{
-				// Talkers needing processing will always be in lsb due to shifts
-				if (DataReadyFlags & 1)
-				{
-					uint32 SpaceAvail = MAX_VOICE_DATA_SIZE - VoiceData.LocalPackets[Index].Length;
-					// Figure out if there is space for this packet
-					if (SpaceAvail > 0)
-					{
-						// Figure out where to append the data
-						uint8* BufferStart = VoiceData.LocalPackets[Index].Buffer.GetData();
-						BufferStart += VoiceData.LocalPackets[Index].Length;
-						// Copy the sender info
-						VoiceData.LocalPackets[Index].Sender = IdentityInt->GetUniquePlayerId(Index);
-						// Process this user
-						uint32 Result = VoiceEngine->ReadLocalVoiceData(Index, BufferStart, &SpaceAvail);
-						if (Result == S_OK)
-						{
-							if (LocalTalkers[Index].bHasNetworkedVoice && SpaceAvail <= MAX_VOICE_PACKET_SIZE_STEAM)
-							{
-								// Mark the person as talking
-								LocalTalkers[Index].bIsTalking = true;
-								LocalTalkers[Index].LastNotificationTime = VoiceNotificationDelta;
-
-								// Update the length based on what it copied
-								VoiceData.LocalPackets[Index].Length += SpaceAvail;
-
-#if VOICE_LOOPBACK
-								if (OSSConsoleVariables::CVarVoiceLoopback.GetValueOnGameThread() && SpaceAvail > 0)
-								{
-									VoiceData.RemotePackets.Add(MakeShareable(new FVoicePacketSteam(VoiceData.LocalPackets[Index])));
-								}
-#endif
-							}
-							else
-							{
-								// Zero out the data since it isn't to be sent via the network
-								VoiceData.LocalPackets[Index].Length = 0;
-							}
-						}
-						else
-						{
-							UE_LOG(LogVoiceEngine, Warning, TEXT("Voice data error in ReadLocalVoiceData"));
-						}
-					}
-					else
-					{
-						UE_LOG(LogVoiceEngine, Warning, TEXT("Dropping voice data due to network layer not processing fast enough"));
-						// Buffer overflow, so drop previous data
-						VoiceData.LocalPackets[Index].Length = 0;
-					}
-				}
-			}
-		}
-	}
-}
-
-void FOnlineVoiceSteam::ProcessRemoteVoicePackets()
-{
-	// Clear the talking state for remote players
-	for (int32 Index = 0; Index < RemoteTalkers.Num(); Index++)
-	{
-		RemoteTalkers[Index].bIsTalking = false;
-	}
-
-	// Now process all pending packets from the server
-	for (int32 Index = 0; Index < VoiceData.RemotePackets.Num(); Index++)
-	{
-		TSharedPtr<FVoicePacketSteam> VoicePacket = StaticCastSharedPtr<FVoicePacketSteam>(VoiceData.RemotePackets[Index]);
-		if (VoicePacket.IsValid())
-		{
-			// Skip local submission of voice if dedicated server or no voice
-			if (VoiceEngine.IsValid())
-			{
-				// Get the size since it is an in/out param
-				uint32 VoiceBufferSize = VoicePacket->GetBufferSize();
-				// Submit this packet to the voice engine
-				uint32 Result = VoiceEngine->SubmitRemoteVoiceData(*VoicePacket->Sender, VoicePacket->Buffer.GetData(), &VoiceBufferSize);
-				if (Result != S_OK)
-				{
-					UE_LOG(LogVoiceEngine, Log,
-						TEXT("SubmitRemoteVoiceData(%s) failed with 0x%08X"),
-						*VoicePacket->Sender->ToDebugString(),
-						Result);
-				}
-			}
-
-			// Find the remote talker and mark them as talking
-			for (int32 Index2 = 0; Index2 < RemoteTalkers.Num(); Index2++)
-			{
-				FRemoteTalker& Talker = RemoteTalkers[Index2];
-				// Compare the ids
-				if (*Talker.TalkerId == *VoicePacket->Sender)
-				{
-					// If the player is marked as muted, they can't be talking
-					Talker.bIsTalking = !IsLocallyMuted(*Talker.TalkerId);
-					Talker.LastNotificationTime = VoiceNotificationDelta;
-				}
-			}
-		}
-	}
-	// Zero the list without causing a free/realloc
-	VoiceData.RemotePackets.Reset();
-}
-
-FString FOnlineVoiceSteam::GetVoiceDebugState() const
-{
-	TSharedPtr<const FUniqueNetId> UniqueId;
-
-	FString Output = TEXT("Voice state\n");
-	Output += VoiceEngine.IsValid() ? VoiceEngine->GetVoiceDebugState() : TEXT("No Voice Engine!");
-
-	Output += TEXT("\nLocal Talkers:\n");
-	for (int32 idx=0; idx < LocalTalkers.Num(); idx++)
-	{
-		UniqueId = IdentityInt->GetUniquePlayerId(idx);
-		
-		const FLocalTalker& Talker = LocalTalkers[idx];
-		Output += FString::Printf(TEXT("[%d]: %s\n Registered: %d\n Networked: %d\n Talking: %d\n WasTalking: %d\n Last:%0.2f\n"),
-			idx,
-			UniqueId.IsValid() ? *UniqueId->ToDebugString() : TEXT("NULL"), 
-			Talker.bIsRegistered,
-			Talker.bHasNetworkedVoice,
-			Talker.bIsTalking,
-			Talker.bWasTalking,
-			Talker.LastNotificationTime);
-	}
-
-	Output += TEXT("\nRemote Talkers:\n");
-	for (int32 idx=0; idx < RemoteTalkers.Num(); idx++)
-	{
-		const FRemoteTalker& Talker = RemoteTalkers[idx];
-		Output += FString::Printf(TEXT("[%d]: %s\n Talking: %d\n WasTalking: %d\n Muted: %s\n Last:%0.2f\n"),
-			idx,
-			*Talker.TalkerId->ToDebugString(), 
-			Talker.bIsTalking,
-			Talker.bWasTalking,
-			IsLocallyMuted(*Talker.TalkerId) ? TEXT("1") : TEXT("0"),
-			Talker.LastNotificationTime);
-
-	}
-
-	Output += TEXT("\nRaw SystemMutelist:\n");
-	for (int32 idx=0; idx < SystemMuteList.Num(); idx++)
-	{
-		Output += FString::Printf(TEXT("[%d]=%s\n"),
-			idx,
-			*SystemMuteList[idx].ToString());
-	}
-
-	Output += TEXT("\nRaw Mutelist:\n");
-	for (int32 idx=0; idx < MuteList.Num(); idx++)
-	{
-		Output += FString::Printf(TEXT("[%d]=%s\n"),
-			idx,
-			*MuteList[idx].ToString());
-	}
-
-	return Output;
-}
-
diff --git a/Engine/Plugins/Online/OnlineSubsystemSteam/Source/Private/VoiceInterfaceSteam.h b/Engine/Plugins/Online/OnlineSubsystemSteam/Source/Private/VoiceInterfaceSteam.h
--- a/Engine/Plugins/Online/OnlineSubsystemSteam/Source/Private/VoiceInterfaceSteam.h
+++ b/Engine/Plugins/Online/OnlineSubsystemSteam/Source/Private/VoiceInterfaceSteam.h
@@ -5,142 +5,39 @@
 #include "CoreMinimal.h"
 #include "Interfaces/VoiceInterface.h"
 #include "Net/VoiceDataCommon.h"
-#include "VoicePacketSteam.h"
+#include "VoicePacketImpl.h"
 #include "OnlineSubsystemSteamTypes.h"
 #include "OnlineSubsystemSteamPackage.h"
+#include "OnlineSubsystemUtilsPackage.h"
+#include "VoiceInterfaceImpl.h"
+#include "VoiceEngineSteam.h"
 
 /**
  * The Steam implementation of the voice interface 
  */
-class FOnlineVoiceSteam : public IOnlineVoice
-{
-	/** Reference to the main Steam subsystem */
-	class FOnlineSubsystemSteam* SteamSubsystem;
-	/** Reference to the sessions interface */
-	class FOnlineSessionSteam* SessionInt;
-	/** Reference to the profile interface */
-	class FOnlineIdentitySteam* IdentityInt;
-	/** Reference to the voice engine for acquiring voice data */
-	IVoiceEnginePtr VoiceEngine;
-
-	/** Maximum permitted local talkers */
-	int32 MaxLocalTalkers;
-	/** Maximum permitted remote talkers */
-	int32 MaxRemoteTalkers;
-
-	/** State of all possible local talkers */
-	TArray<FLocalTalker> LocalTalkers;
-	/** State of all possible remote talkers */
-	TArray<FRemoteTalker> RemoteTalkers;
-	/** Remote players locally muted explicitly */
-	TArray<FUniqueNetIdSteam> SystemMuteList;
-	/** Remote players locally muted (super set of SystemMuteList) */
-	TArray<FUniqueNetIdSteam> MuteList;
-
-	/** Time to wait for new data before triggering "not talking" */
-	float VoiceNotificationDelta;
-
-	/** Buffered voice data I/O */
-	FVoiceDataSteam VoiceData;
-
-	/**
-	 * Finds a remote talker in the cached list
-	 *
-	 * @param UniqueId the net id of the player to search for
-	 *
-	 * @return pointer to the remote talker or NULL if not found
-	 */
-	struct FRemoteTalker* FindRemoteTalker(const FUniqueNetId& UniqueId);
-
-	/**
-	 * Is a given id presently muted (either by system mute or game server)
-	 *
-	 * @param UniqueId the net id to query
-	 *
-	 * @return true if the net id is muted at all, false otherwise
-	 */
-	bool IsLocallyMuted(const FUniqueNetId& UniqueId) const;
-
-	/**
-	 * Does a given id exist in the system wide mute list
-	 *
-	 * @param UniqueId the net id to query
-	 *
-	 * @return true if the net id is on the system wide mute list, false otherwise
-	 */
-	bool IsSystemWideMuted(const FUniqueNetId& UniqueId) const;
 
+class ONLINESUBSYSTEMSTEAM_API FOnlineVoiceSteam : public FOnlineVoiceImpl
+{
 PACKAGE_SCOPE:
-	FOnlineVoiceSteam() :
-		SteamSubsystem(NULL),
-		SessionInt(NULL),
-		IdentityInt(NULL),
-		VoiceEngine(NULL),
-		MaxLocalTalkers(MAX_SPLITSCREEN_TALKERS),
-		MaxRemoteTalkers(MAX_REMOTE_TALKERS),
-		VoiceNotificationDelta(0.0f)
+	FOnlineVoiceSteam() : FOnlineVoiceImpl()
 	{};
 
-	// IOnlineVoice
-	virtual bool Init() override;
-	void ProcessMuteChangeNotification() override;
-
-	/**
-	 * Processes any talking delegates that need to be fired off
-	 *
-	 * @param DeltaTime the amount of time that has elapsed since the last tick
-	 */
-	void ProcessTalkingDelegates(float DeltaTime);
-
-	/**
-	 * Reads any data that is currently queued
-	 */
-	void ProcessLocalVoicePackets();
-
-	/**
-	 * Submits network packets to audio system for playback
-	 */
-	void ProcessRemoteVoicePackets();
-
-	/**
-	 * Figures out which remote talkers need to be muted for a given local talker
-	 *
-	 * @param TalkerIndex the talker that needs the mute list checked for
-	 * @param PlayerController the player controller associated with this talker
-	 */
-	void UpdateMuteListForLocalTalker(int32 TalkerIndex, class APlayerController* PlayerController);
-
 public:
 
 	/** Constructor */
-	FOnlineVoiceSteam(class FOnlineSubsystemSteam* InSteamSubsystem);
+	FOnlineVoiceSteam(class IOnlineSubsystem* InOnlineSubsystem) :
+		FOnlineVoiceImpl(InOnlineSubsystem)
+	{
+		check(InOnlineSubsystem);
+	};
 
-	/** Virtual destructor to force proper child cleanup */
-	virtual ~FOnlineVoiceSteam();
+	virtual IVoiceEnginePtr CreateVoiceEngine() override
+	{
+		return MakeShareable(new FVoiceEngineSteam(OnlineSubsystem));
+	}
 
-	// IOnlineVoice
-	virtual void StartNetworkedVoice(uint8 LocalUserNum) override;
-	virtual void StopNetworkedVoice(uint8 LocalUserNum) override;
-    virtual bool RegisterLocalTalker(uint32 LocalUserNum) override;
-	virtual void RegisterLocalTalkers() override;
-    virtual bool UnregisterLocalTalker(uint32 LocalUserNum) override;
-	virtual void UnregisterLocalTalkers() override;
-    virtual bool RegisterRemoteTalker(const FUniqueNetId& UniqueId) override;
-    virtual bool UnregisterRemoteTalker(const FUniqueNetId& UniqueId) override;
-	virtual void RemoveAllRemoteTalkers() override;
-    virtual bool IsHeadsetPresent(uint32 LocalUserNum) override;
-    virtual bool IsLocalPlayerTalking(uint32 LocalUserNum) override;
-	virtual bool IsRemotePlayerTalking(const FUniqueNetId& UniqueId) override;
-	bool IsMuted(uint32 LocalUserNum, const FUniqueNetId& UniqueId) const override;
-	bool MuteRemoteTalker(uint8 LocalUserNum, const FUniqueNetId& PlayerId, bool bIsSystemWide) override;
-	bool UnmuteRemoteTalker(uint8 LocalUserNum, const FUniqueNetId& PlayerId, bool bIsSystemWide) override;
-	virtual TSharedPtr<class FVoicePacket> SerializeRemotePacket(FArchive& Ar) override;
-	virtual TSharedPtr<class FVoicePacket> GetLocalPacket(uint32 LocalUserNum) override;
-	virtual int32 GetNumLocalTalkers() override { return LocalTalkers.Num(); };
-	virtual void ClearVoicePackets() override;
-	virtual void Tick(float DeltaTime) override;
-	virtual FString GetVoiceDebugState() const override;
+	/** Virtual destructor to force proper child cleanup */
+	virtual ~FOnlineVoiceSteam() override {};
 };
 
 typedef TSharedPtr<FOnlineVoiceSteam, ESPMode::ThreadSafe> FOnlineVoiceSteamPtr;
-
diff --git a/Engine/Plugins/Online/OnlineSubsystemSteam/Source/Public/OnlineSubsystemSteam.h b/Engine/Plugins/Online/OnlineSubsystemSteam/Source/Public/OnlineSubsystemSteam.h
--- a/Engine/Plugins/Online/OnlineSubsystemSteam/Source/Public/OnlineSubsystemSteam.h
+++ b/Engine/Plugins/Online/OnlineSubsystemSteam/Source/Public/OnlineSubsystemSteam.h
@@ -79,7 +79,10 @@ protected:
 	FOnlineLeaderboardsSteamPtr LeaderboardsInterface;
 
 	/** Interface to the voice engine */
-	FOnlineVoiceSteamPtr VoiceInterface;
+	mutable IOnlineVoicePtr VoiceInterface;
+
+	/** Interface for voice communication */
+	mutable bool bVoiceInterfaceInitialized;
 
 	/** Interface to the external UI services */
 	FOnlineExternalUISteamPtr ExternalUIInterface;
@@ -117,6 +120,7 @@ PACKAGE_SCOPE:
 		UserCloudInterface(nullptr),
 		LeaderboardsInterface(nullptr),
 		VoiceInterface(nullptr),
+		bVoiceInterfaceInitialized(false),
 		ExternalUIInterface(nullptr),
 		PresenceInterface(nullptr),
 		AuthInterface(nullptr),
@@ -138,6 +142,7 @@ PACKAGE_SCOPE:
 		UserCloudInterface(nullptr),
 		LeaderboardsInterface(nullptr),
 		VoiceInterface(nullptr),
+		bVoiceInterfaceInitialized(false),
 		ExternalUIInterface(nullptr),
 		PresenceInterface(nullptr),
 		AuthInterface(nullptr),
diff --git a/Engine/Plugins/Online/OnlineSubsystemUtils/Source/OnlineSubsystemUtils/Private/VoiceEngineImpl.cpp b/Engine/Plugins/Online/OnlineSubsystemUtils/Source/OnlineSubsystemUtils/Private/VoiceEngineImpl.cpp
--- a/Engine/Plugins/Online/OnlineSubsystemUtils/Source/OnlineSubsystemUtils/Private/VoiceEngineImpl.cpp
+++ b/Engine/Plugins/Online/OnlineSubsystemUtils/Source/OnlineSubsystemUtils/Private/VoiceEngineImpl.cpp
@@ -286,13 +286,12 @@ uint32 FVoiceEngineImpl::StopLocalVoiceProcessing(uint32 LocalUserNum)
 
 uint32 FVoiceEngineImpl::UnregisterRemoteTalker(const FUniqueNetId& UniqueId)
 {
-	const FUniqueNetIdString& RemoteTalkerId = (const FUniqueNetIdString&)UniqueId;
-	FRemoteTalkerDataImpl* RemoteData = RemoteTalkerBuffers.Find(RemoteTalkerId);
+	FRemoteTalkerDataImpl* RemoteData = RemoteTalkerBuffers.Find(FUniqueNetIdWrapper(UniqueId.AsShared()));
 	if (RemoteData != nullptr)
 	{
 		// Dump the whole talker
 		RemoteData->Cleanup();
-		RemoteTalkerBuffers.Remove(RemoteTalkerId);
+		RemoteTalkerBuffers.Remove(FUniqueNetIdWrapper(UniqueId.AsShared()));
 	}
 
 	return S_OK;
@@ -367,7 +366,7 @@ uint32 FVoiceEngineImpl::ReadLocalVoiceData(uint32 LocalUserNum, uint8* Data, ui
 		uint8* RemainingDecompressedBufferPtr = DecompressedVoiceBuffer.GetData() + PlayerVoiceData[LocalUserNum].VoiceRemainderSize;
 		uint32 ByteWritten = 0;
 		uint64 NewSampleCount = 0;
-		VoiceResult = VoiceCapture->GetVoiceData(DecompressedVoiceBuffer.GetData() + PlayerVoiceData[LocalUserNum].VoiceRemainderSize, NewVoiceDataBytes, ByteWritten , NewSampleCount);
+		VoiceResult = VoiceCapture->GetVoiceData(DecompressedVoiceBuffer.GetData() + PlayerVoiceData[LocalUserNum].VoiceRemainderSize, NewVoiceDataBytes, ByteWritten, NewSampleCount);
 		
 		TotalVoiceBytes = ByteWritten + PlayerVoiceData[LocalUserNum].VoiceRemainderSize;
 
@@ -382,7 +381,7 @@ uint32 FVoiceEngineImpl::ReadLocalVoiceData(uint32 LocalUserNum, uint8* Data, ui
 			CompressedBytesAvailable = UVOIPStatics::GetMaxCompressedVoiceDataSize();
 			CompressedVoiceBuffer.AddUninitialized(UVOIPStatics::GetMaxCompressedVoiceDataSize());
 
-			check( ((uint32) CompressedVoiceBuffer.Num()) <= UVOIPStatics::GetMaxCompressedVoiceDataSize());
+			check(((uint32) CompressedVoiceBuffer.Num()) <= UVOIPStatics::GetMaxCompressedVoiceDataSize());
 
 			// Run the uncompressed audio through the opus decoder, note that it may not encode all data, which results in some remaining data
 			PlayerVoiceData[LocalUserNum].VoiceRemainderSize =
@@ -433,8 +432,7 @@ uint32 FVoiceEngineImpl::SubmitRemoteVoiceData(const FUniqueNetIdWrapper& Remote
 {
 	UE_LOG_ONLINE_VOICEENGINE(VeryVerbose, TEXT("SubmitRemoteVoiceData(%s) Size: %d received!"), *RemoteTalkerId.ToDebugString(), *Size);
 	
-	const FUniqueNetIdString& TalkerId = (const FUniqueNetIdString&) (*RemoteTalkerId);
-	FRemoteTalkerDataImpl& QueuedData = RemoteTalkerBuffers.FindOrAdd(TalkerId);
+	FRemoteTalkerDataImpl& QueuedData = RemoteTalkerBuffers.FindOrAdd(RemoteTalkerId);
 
 	// new voice packet.
 	QueuedData.LastSeen = FPlatformTime::Seconds();
@@ -479,9 +477,8 @@ uint32 FVoiceEngineImpl::SubmitRemoteVoiceData(const FUniqueNetIdWrapper& Remote
 			FVoiceSettings InSettings;
 			UVOIPTalker* OwningTalker = nullptr;
 
-			
 			OwningTalker = UVOIPStatics::GetVOIPTalkerForPlayer(RemoteTalkerId, InSettings);
-			
+
 			ApplyVoiceSettings(QueuedData.VoipSynthComponent, InSettings);
 
 			QueuedData.VoipSynthComponent->ResetBuffer(InSampleCount, UVOIPStatics::GetBufferingDelay());
@@ -500,7 +497,7 @@ uint32 FVoiceEngineImpl::SubmitRemoteVoiceData(const FUniqueNetIdWrapper& Remote
 			}
 		}
 
-		QueuedData.VoipSynthComponent->SubmitPacket((float*) DecompressedVoiceBuffer.GetData(), BytesWritten, InSampleCount, EVoipStreamDataFormat::Int16);
+		QueuedData.VoipSynthComponent->SubmitPacket((float*)DecompressedVoiceBuffer.GetData(), BytesWritten, InSampleCount, EVoipStreamDataFormat::Int16);
 	}
 
 	return S_OK;
@@ -550,9 +547,9 @@ void FVoiceEngineImpl::Tick(float DeltaTime)
 	TickTalkers(DeltaTime);
 }
 
-void FVoiceEngineImpl::GenerateVoiceData(USoundWaveProcedural* InProceduralWave, int32 SamplesRequired, FUniqueNetIdString TalkerId)
+void FVoiceEngineImpl::GenerateVoiceData(USoundWaveProcedural* InProceduralWave, int32 SamplesRequired, const FUniqueNetId& TalkerId)
 {
-	FRemoteTalkerDataImpl* QueuedData = RemoteTalkerBuffers.Find(TalkerId);
+	FRemoteTalkerDataImpl* QueuedData = RemoteTalkerBuffers.Find(FUniqueNetIdWrapper(TalkerId.AsShared()));
 	if (QueuedData)
 	{
 		const int32 SampleSize = sizeof(uint16) * DEFAULT_NUM_VOICE_CHANNELS;
diff --git a/Engine/Plugins/Online/OnlineSubsystemUtils/Source/OnlineSubsystemUtils/Private/VoiceInterfaceImpl.cpp b/Engine/Plugins/Online/OnlineSubsystemUtils/Source/OnlineSubsystemUtils/Private/VoiceInterfaceImpl.cpp
--- a/Engine/Plugins/Online/OnlineSubsystemUtils/Source/OnlineSubsystemUtils/Private/VoiceInterfaceImpl.cpp
+++ b/Engine/Plugins/Online/OnlineSubsystemUtils/Source/OnlineSubsystemUtils/Private/VoiceInterfaceImpl.cpp
@@ -61,7 +61,7 @@ bool FOnlineVoiceImpl::Init()
 			const bool bVoiceEngineForceDisable = OnlineSubsystem->IsDedicated() || GIsBuildMachine;
 			if (!bVoiceEngineForceDisable)
 			{
-				VoiceEngine = MakeShareable(new FVoiceEngineImpl(OnlineSubsystem));
+				VoiceEngine = CreateVoiceEngine();
 				bSuccess = VoiceEngine->Init(MaxLocalTalkers, MaxRemoteTalkers);
 			}
 			else
@@ -117,7 +117,7 @@ void FOnlineVoiceImpl::ClearVoicePackets()
 	}
 }
 
-void FOnlineVoiceImpl::Tick(float DeltaTime) 
+void FOnlineVoiceImpl::Tick(float DeltaTime)
 {
 	if (!OnlineSubsystem->IsDedicated())
 	{
@@ -182,7 +182,7 @@ void FOnlineVoiceImpl::StopNetworkedVoice(uint8 LocalUserNum)
 	}
 }
 
-bool FOnlineVoiceImpl::RegisterLocalTalker(uint32 LocalUserNum) 
+bool FOnlineVoiceImpl::RegisterLocalTalker(uint32 LocalUserNum)
 {
 	uint32 Return = E_FAIL;
 	if (LocalUserNum >= 0 && LocalUserNum < (uint32)MaxLocalTalkers)
@@ -243,7 +243,7 @@ void FOnlineVoiceImpl::RegisterLocalTalkers()
 	}
 }
 
-bool FOnlineVoiceImpl::UnregisterLocalTalker(uint32 LocalUserNum) 
+bool FOnlineVoiceImpl::UnregisterLocalTalker(uint32 LocalUserNum)
 {
 	uint32 Return = S_OK;
 	if (LocalUserNum >= 0 && LocalUserNum < (uint32)MaxLocalTalkers)
@@ -296,7 +296,7 @@ void FOnlineVoiceImpl::UnregisterLocalTalkers()
 	}
 }
 
-bool FOnlineVoiceImpl::RegisterRemoteTalker(const FUniqueNetId& UniqueId) 
+bool FOnlineVoiceImpl::RegisterRemoteTalker(const FUniqueNetId& UniqueId)
 {
 	uint32 Return = E_FAIL;
 	if (OnlineSubsystem)
@@ -314,8 +314,8 @@ bool FOnlineVoiceImpl::RegisterRemoteTalker(const FUniqueNetId& UniqueId)
 				int32 AddIndex = RemoteTalkers.AddZeroed();
 				Talker = &RemoteTalkers[AddIndex];
 				// Copy the UniqueId
-				const FUniqueNetIdString& UniqueIdStr = (const FUniqueNetIdString&)UniqueId;
-				Talker->TalkerId = MakeShareable(new FUniqueNetIdString(UniqueIdStr));
+				Talker->TalkerId = UniqueId.AsShared();
+
 				// Register the remote talker locally
 				Return = VoiceEngine->RegisterRemoteTalker(UniqueId);
 				UE_LOG_ONLINE_VOICE(Log, TEXT("RegisterRemoteTalker(%s) returned 0x%08X"),
@@ -337,7 +337,7 @@ bool FOnlineVoiceImpl::RegisterRemoteTalker(const FUniqueNetId& UniqueId)
 	return Return == S_OK;
 }
 
-bool FOnlineVoiceImpl::UnregisterRemoteTalker(const FUniqueNetId& UniqueId) 
+bool FOnlineVoiceImpl::UnregisterRemoteTalker(const FUniqueNetId& UniqueId)
 {
 	uint32 Return = E_FAIL;
 	if (OnlineSubsystem)
@@ -418,17 +418,17 @@ FRemoteTalker* FOnlineVoiceImpl::FindRemoteTalker(const FUniqueNetId& UniqueId)
 	return NULL;
 }
 
-bool FOnlineVoiceImpl::IsHeadsetPresent(uint32 LocalUserNum) 
+bool FOnlineVoiceImpl::IsHeadsetPresent(uint32 LocalUserNum)
 {
 	return VoiceEngine.IsValid() && VoiceEngine->IsHeadsetPresent(LocalUserNum);
 }
 
-bool FOnlineVoiceImpl::IsLocalPlayerTalking(uint32 LocalUserNum) 
+bool FOnlineVoiceImpl::IsLocalPlayerTalking(uint32 LocalUserNum)
 {
 	return VoiceEngine.IsValid() && VoiceEngine->IsLocalPlayerTalking(LocalUserNum);
 }
 
-bool FOnlineVoiceImpl::IsRemotePlayerTalking(const FUniqueNetId& UniqueId) 
+bool FOnlineVoiceImpl::IsRemotePlayerTalking(const FUniqueNetId& UniqueId)
 {
 	return VoiceEngine.IsValid() && VoiceEngine->IsRemotePlayerTalking(UniqueId);
 }
@@ -445,13 +445,13 @@ bool FOnlineVoiceImpl::IsMuted(uint32 LocalUserNum, const FUniqueNetId& UniqueId
 
 bool FOnlineVoiceImpl::IsLocallyMuted(const FUniqueNetId& UniqueId) const
 {
-	int32 Index = MuteList.Find((const FUniqueNetIdString&)UniqueId);
+	int32 Index = MuteList.Find(FUniqueNetIdWrapper(UniqueId.AsShared()));
 	return Index != INDEX_NONE;
 }
 
 bool FOnlineVoiceImpl::IsSystemWideMuted(const FUniqueNetId& UniqueId) const
 {
-	int32 Index = SystemMuteList.Find((const FUniqueNetIdString&)UniqueId);
+	int32 Index = SystemMuteList.Find(FUniqueNetIdWrapper(UniqueId.AsShared()));
 	return Index != INDEX_NONE;
 }
 
@@ -463,7 +463,7 @@ bool FOnlineVoiceImpl::MuteRemoteTalker(uint8 LocalUserNum, const FUniqueNetId&
 		if (bIsSystemWide)
 		{
 			// Add them to the system wide mute list
-			SystemMuteList.AddUnique((const FUniqueNetIdString&)PlayerId);
+			SystemMuteList.AddUnique(FUniqueNetIdWrapper(PlayerId.AsShared()));
 			// Should update MuteList after going up to the server and coming back down
 			ProcessMuteChangeNotification();
 		}
@@ -478,7 +478,7 @@ bool FOnlineVoiceImpl::MuteRemoteTalker(uint8 LocalUserNum, const FUniqueNetId&
 				FRemoteTalker* Talker = FindRemoteTalker(PlayerId);
 				if (Talker != NULL)
 				{
-					MuteList.AddUnique((const FUniqueNetIdString&)PlayerId);
+					MuteList.AddUnique(FUniqueNetIdWrapper(PlayerId.AsShared()));
 					Return = S_OK;
 					UE_LOG_ONLINE_VOICE(Log, TEXT("Muting remote talker (%s)"), *PlayerId.ToDebugString());
 				}
@@ -505,7 +505,7 @@ bool FOnlineVoiceImpl::UnmuteRemoteTalker(uint8 LocalUserNum, const FUniqueNetId
 		if (bIsSystemWide)
 		{
 			// Remove them from the system wide mute list
-			SystemMuteList.RemoveSingleSwap((const FUniqueNetIdString&)PlayerId);
+			SystemMuteList.RemoveSingleSwap(FUniqueNetIdWrapper(PlayerId.AsShared()));
 			// Should update MuteList after going up to the server and coming back down
 			ProcessMuteChangeNotification();
 		}
@@ -525,7 +525,7 @@ bool FOnlineVoiceImpl::UnmuteRemoteTalker(uint8 LocalUserNum, const FUniqueNetId
 					if (!bIsSystemMuted)
 					{
 						// Remove them from the mute list
-						MuteList.RemoveSingleSwap((const FUniqueNetIdString&)PlayerId);
+						MuteList.RemoveSingleSwap(FUniqueNetIdWrapper(PlayerId.AsShared()));
 						UE_LOG_ONLINE_VOICE(Log, TEXT("Unmuting remote talker (%s)"), *PlayerId.ToDebugString());
 					}
 				}
@@ -572,6 +572,11 @@ void FOnlineVoiceImpl::ProcessMuteChangeNotification()
 	}
 }
 
+IVoiceEnginePtr FOnlineVoiceImpl::CreateVoiceEngine()
+{
+	return MakeShareable(new FVoiceEngineImpl(OnlineSubsystem));
+}
+
 void FOnlineVoiceImpl::UpdateMuteListForLocalTalker(int32 TalkerIndex, APlayerController* PlayerController)
 {
 	// For each registered remote talker
@@ -582,7 +587,7 @@ void FOnlineVoiceImpl::UpdateMuteListForLocalTalker(int32 TalkerIndex, APlayerCo
 		FUniqueNetIdRepl UniqueIdRepl(Talker.TalkerId);
 
 		// Is the remote talker on this local player's mute list?
-		if (SystemMuteList.Find((FUniqueNetIdString&)*Talker.TalkerId) == INDEX_NONE)
+		if (SystemMuteList.Find(FUniqueNetIdWrapper(Talker.TalkerId->AsShared())) == INDEX_NONE)
 		{
 			// Unmute on the server
 			PlayerController->ServerUnmutePlayer(UniqueIdRepl);
@@ -606,13 +611,11 @@ TSharedPtr<FVoicePacket> FOnlineVoiceImpl::SerializeRemotePacket(FArchive& Ar)
 	{
 		if (!OnlineSubsystem->IsDedicated())
 		{
-			FUniqueNetIdMatcher PlayerMatch(*NewPacket->GetSender());
-			if (MuteList.IndexOfByPredicate(PlayerMatch) == INDEX_NONE)
+			if (MuteList.Find(FUniqueNetIdWrapper( NewPacket->GetSender())) == INDEX_NONE)
 			{
 				VoiceData.RemotePackets.Add(NewPacket);
 			}
 		}
-
 		return NewPacket;
 	}
 
@@ -883,4 +886,3 @@ FString FOnlineVoiceImpl::GetVoiceDebugState() const
 
 	return Output;
 }
-
