Merge Online changes to Engine/Ocean/OrionGame/UnrealTournament/WEX.

diff --git a/Engine/Plugins/Online/OnlineSubsystemOculus/Source/Private/IPAddressOculus.h b/Engine/Plugins/Online/OnlineSubsystemOculus/Source/Private/IPAddressOculus.h
--- a/Engine/Plugins/Online/OnlineSubsystemOculus/Source/Private/IPAddressOculus.h
+++ b/Engine/Plugins/Online/OnlineSubsystemOculus/Source/Private/IPAddressOculus.h
@@ -156,13 +156,6 @@ public:
 		return !(FInternetAddrOculus::operator==(Other));
 	}
 
-	virtual bool operator<(const FInternetAddr& Other) const override
-	{
-		FInternetAddrOculus& OculusOther = (FInternetAddrOculus&)Other;
-
-		return GetID() < OculusOther.GetID();
-	}
-
 	virtual uint32 GetTypeHash() override
 	{
 		return ::GetTypeHash((uint64)GetID());
diff --git a/Engine/Plugins/Online/OnlineSubsystemSteam/Source/Private/IPAddressSteam.h b/Engine/Plugins/Online/OnlineSubsystemSteam/Source/Private/IPAddressSteam.h
--- a/Engine/Plugins/Online/OnlineSubsystemSteam/Source/Private/IPAddressSteam.h
+++ b/Engine/Plugins/Online/OnlineSubsystemSteam/Source/Private/IPAddressSteam.h
@@ -144,13 +144,6 @@ public:
 		return !(FInternetAddrSteam::operator==(Other));
 	}
 
-	virtual bool operator<(const FInternetAddr& Other) const override
-	{
-		FInternetAddrSteam& SteamOther = (FInternetAddrSteam&)Other;
-
-		return (CSteamID)SteamId < (CSteamID)SteamOther.SteamId || SteamChannel < SteamOther.SteamChannel;
-	}
-
 	virtual uint32 GetTypeHash() override
 	{
 		return ::GetTypeHash(*(uint64*)SteamId.GetBytes());
