Fixing up the spammy logs due to kick attempts on fail auth.

Also making the reauth message more clear.

diff --git a/Engine/Plugins/Online/OnlineSubsystemSteam/Source/Private/OnlineAuthInterfaceSteam.cpp b/Engine/Plugins/Online/OnlineSubsystemSteam/Source/Private/OnlineAuthInterfaceSteam.cpp
--- a/Engine/Plugins/Online/OnlineSubsystemSteam/Source/Private/OnlineAuthInterfaceSteam.cpp
+++ b/Engine/Plugins/Online/OnlineSubsystemSteam/Source/Private/OnlineAuthInterfaceSteam.cpp
@@ -132,7 +132,7 @@ bool FOnlineAuthSteam::AuthenticateUser(const FUniqueNetId& InUserId)
 		// Do not attempt to reauth this user if we are currently doing this.
 		if (EnumHasAnyFlags(TargetUser->Status, ESteamAuthStatus::HasOrIsPendingAuth))
 		{
-			UE_LOG_ONLINE(Log, TEXT("AUTH: Currently authenticating with user %s, skipping reauth"), *InUserId.ToString());
+			UE_LOG_ONLINE(Log, TEXT("AUTH: The user %s has authenticated or is currently authenticating. Skipping reauth"), *InUserId.ToString());
 			return true;
 		}
 
@@ -146,7 +146,7 @@ bool FOnlineAuthSteam::AuthenticateUser(const FUniqueNetId& InUserId)
 		if (TargetUser->RecvTicket.IsEmpty())
 		{
 			UE_LOG_ONLINE(Warning, TEXT("AUTH: Ticket from user %s is empty"), *InUserId.ToString());
-			TargetUser->Status |= ESteamAuthStatus::FailKick;
+			TargetUser->Status |= ESteamAuthStatus::AuthFail;
 			return false;
 		}
 
@@ -154,7 +154,7 @@ bool FOnlineAuthSteam::AuthenticateUser(const FUniqueNetId& InUserId)
 		if (TargetUser->RecvTicket.Len() > STEAM_AUTH_MAX_TICKET_LENGTH_IN_BYTES)
 		{
 			UE_LOG_ONLINE(Warning, TEXT("AUTH: Ticket from user is over max size of ticket length"));
-			TargetUser->Status |= ESteamAuthStatus::FailKick;
+			TargetUser->Status |= ESteamAuthStatus::AuthFail;
 			return false;
 		}
 
@@ -164,7 +164,7 @@ bool FOnlineAuthSteam::AuthenticateUser(const FUniqueNetId& InUserId)
 			if (!CheckTCharIsHex(TargetUser->RecvTicket.GetCharArray()[i]))
 			{
 				UE_LOG_ONLINE(Warning, TEXT("AUTH: Ticket from user is not stored in hex!"));
-				TargetUser->Status |= ESteamAuthStatus::FailKick;
+				TargetUser->Status |= ESteamAuthStatus::AuthFail;
 				return false;
 			}
 		}
@@ -189,7 +189,7 @@ bool FOnlineAuthSteam::AuthenticateUser(const FUniqueNetId& InUserId)
 			else
 			{
 				UE_LOG_ONLINE(Warning, TEXT("AUTH: User %s failed authentication %d"), *InUserId.ToString(), (int32)Result);
-				TargetUser->Status |= ESteamAuthStatus::FailKick;
+				TargetUser->Status |= ESteamAuthStatus::AuthFail;
 			}
 		}
 		else
@@ -205,7 +205,7 @@ bool FOnlineAuthSteam::AuthenticateUser(const FUniqueNetId& InUserId)
 			else
 			{
 				UE_LOG_ONLINE(Warning, TEXT("AUTH: User %s failed authentication %d"), *InUserId.ToString(), (int32)Result);
-				TargetUser->Status |= ESteamAuthStatus::FailKick;
+				TargetUser->Status |= ESteamAuthStatus::AuthFail;
 			}
 		}
 	}
@@ -283,12 +283,12 @@ void FOnlineAuthSteam::MarkPlayerForKick(const FUniqueNetId& InUserId)
 	SharedAuthUserSteamPtr TargetUser = GetUser(SteamId);
 	if (TargetUser.IsValid())
 	{
-		TargetUser->Status |= ESteamAuthStatus::FailKick;
+		TargetUser->Status |= ESteamAuthStatus::AuthFail;
 		UE_LOG_ONLINE(Log, TEXT("AUTH: Marking %s for kick"), *InUserId.ToString());
 	}
 }
 
-bool FOnlineAuthSteam::KickPlayer(const FUniqueNetId& InUserId)
+bool FOnlineAuthSteam::KickPlayer(const FUniqueNetId& InUserId, bool bSuppressFailure)
 {
 	bool bKickSuccess = false;
 	const FUniqueNetIdSteam& SteamId = static_cast<const FUniqueNetIdSteam&>(InUserId);
@@ -296,7 +296,10 @@ bool FOnlineAuthSteam::KickPlayer(const FUniqueNetId& InUserId)
 
 	if (SteamUserPtr != nullptr && SteamUserPtr->GetSteamID() == SteamId)
 	{
-		UE_LOG_ONLINE(Warning, TEXT("AUTH: Cannot kick ourselves!"));
+		if (!bSuppressFailure)
+		{
+			UE_LOG_ONLINE(Warning, TEXT("AUTH: Cannot kick ourselves!"));
+		}
 		return false;
 	}
 
@@ -313,7 +316,10 @@ bool FOnlineAuthSteam::KickPlayer(const FUniqueNetId& InUserId)
 		AGameModeBase* GameMode = World->GetAuthGameMode();
 		if (GameMode == nullptr || GameMode->GameSession == nullptr)
 		{
-			UE_LOG_ONLINE(Warning, TEXT("AUTH: Cannot kick player %s as we do not have a gamemode or session"), *InUserId.ToString());
+			if (!bSuppressFailure)
+			{
+				UE_LOG_ONLINE(Warning, TEXT("AUTH: Cannot kick player %s as we do not have a gamemode or session"), *InUserId.ToString());
+			}
 			return false;
 		}
 
@@ -337,7 +343,7 @@ bool FOnlineAuthSteam::KickPlayer(const FUniqueNetId& InUserId)
 		UE_LOG_ONLINE(Log, TEXT("AUTH: Successfully kicked player %s"), *InUserId.ToString());
 		RemoveUser(InUserId);
 	}
-	else
+	else if(!bSuppressFailure)
 	{
 		UE_LOG_ONLINE(Warning, TEXT("AUTH: Was not able to kick player %s Valid world: %d."), *InUserId.ToString(), (World != nullptr));
 	}
@@ -380,11 +386,12 @@ bool FOnlineAuthSteam::Tick(float DeltaTime)
 			// Kick any players that have failed authentication.
 			if (EnumHasAnyFlags(CurUser->Status, ESteamAuthStatus::FailKick))
 			{
-				if (KickPlayer(CurUserId))
+				if (KickPlayer(CurUserId, EnumHasAnyFlags(CurUser->Status, ESteamAuthStatus::KickUser)))
 				{
 					// If we've modified the list, we can just end this frame.
 					return true;
 				}
+				CurUser->Status |= ESteamAuthStatus::KickUser;
 			}
 		}
 	}
@@ -504,7 +511,7 @@ void FOnlineAuthSteam::OnAuthResult(const FUniqueNetId& TargetId, int32 Response
 	}
 	else
 	{
-		TargetUser->Status |= ESteamAuthStatus::FailKick;
+		TargetUser->Status |= ESteamAuthStatus::AuthFail;
 	}
 	ExecuteResultDelegate(SteamId, bDidAuthSucceed);
 }
diff --git a/Engine/Plugins/Online/OnlineSubsystemSteam/Source/Public/OnlineAuthInterfaceSteam.h b/Engine/Plugins/Online/OnlineSubsystemSteam/Source/Public/OnlineAuthInterfaceSteam.h
--- a/Engine/Plugins/Online/OnlineSubsystemSteam/Source/Public/OnlineAuthInterfaceSteam.h
+++ b/Engine/Plugins/Online/OnlineSubsystemSteam/Source/Public/OnlineAuthInterfaceSteam.h
@@ -33,8 +33,10 @@ enum class ESteamAuthStatus : uint8
 {
 	None = 0,
 	AuthSuccess = 1 << 0,
-	FailKick = 1 << 1,
+	AuthFail = 1 << 1,
 	ValidationStarted = 1 << 2,
+	KickUser = 1 << 3,
+	FailKick = AuthFail | KickUser,
 	HasOrIsPendingAuth = AuthSuccess | ValidationStarted
 };
 
@@ -88,7 +90,7 @@ private:
 	{
 		return SteamSubsystem != nullptr && SteamSubsystem->IsServer();
 	}
-	bool KickPlayer(const FUniqueNetId& InUserId);
+	bool KickPlayer(const FUniqueNetId& InUserId, bool bSuppressFailure);
 
 	FOnlineAuthSteam();
 
