Allow build ids of 0 to not be filtered out incorrectly on Steam.

diff --git a/Engine/Plugins/Online/OnlineSubsystemSteam/Source/Private/OnlineSessionAsyncLobbySteam.cpp b/Engine/Plugins/Online/OnlineSubsystemSteam/Source/Private/OnlineSessionAsyncLobbySteam.cpp
--- a/Engine/Plugins/Online/OnlineSubsystemSteam/Source/Private/OnlineSessionAsyncLobbySteam.cpp
+++ b/Engine/Plugins/Online/OnlineSubsystemSteam/Source/Private/OnlineSessionAsyncLobbySteam.cpp
@@ -266,12 +266,8 @@ bool FillSessionFromLobbyData(FUniqueNetIdSteam& LobbyId, FOnlineSession& Sessio
 		}
 		else if (FCStringAnsi::Stricmp(Key, STEAMKEY_BUILDUNIQUEID) == 0)
 		{
-			int32 BuildUniqueId = FCString::Atoi(ANSI_TO_TCHAR(Value));
-			if (BuildUniqueId != 0)
-			{
-				Session.SessionSettings.BuildUniqueId = BuildUniqueId;
-				KeysFound++;
-			}
+			Session.SessionSettings.BuildUniqueId = FCString::Atoi(ANSI_TO_TCHAR(Value));
+			KeysFound++;
 		}
 		else if (FCStringAnsi::Stricmp(Key, STEAMKEY_OWNINGUSERID) == 0)
 		{
@@ -353,7 +349,7 @@ bool FillSessionFromLobbyData(FUniqueNetIdSteam& LobbyId, FOnlineSession& Sessio
 	if (bSuccess && (KeysFound == STEAMKEY_NUMREQUIREDLOBBYKEYS) && (HostKeysFound == 2 || SteamAddrKeysFound == 2))
 	{
 		int32 BuildUniqueId = GetBuildUniqueId();
-		if (Session.SessionSettings.BuildUniqueId != 0 && Session.SessionSettings.BuildUniqueId == BuildUniqueId)
+		if (Session.SessionSettings.BuildUniqueId == BuildUniqueId)
 		{
 			if (HostKeysFound == 2)
 			{
diff --git a/Engine/Plugins/Online/OnlineSubsystemSteam/Source/Private/OnlineSessionAsyncServerSteam.cpp b/Engine/Plugins/Online/OnlineSubsystemSteam/Source/Private/OnlineSessionAsyncServerSteam.cpp
--- a/Engine/Plugins/Online/OnlineSubsystemSteam/Source/Private/OnlineSessionAsyncServerSteam.cpp
+++ b/Engine/Plugins/Online/OnlineSubsystemSteam/Source/Private/OnlineSessionAsyncServerSteam.cpp
@@ -953,7 +953,7 @@ void FOnlineAsyncTaskSteamFindServerBase::ParseSearchResult(class gameserveritem
 			ServerBuildId = FCString::Atoi(*TagArray[0].Mid(ARRAY_COUNT(STEAMKEY_BUILDUNIQUEID)));
 		}
 
-		if (ServerBuildId != 0 && ServerBuildId == BuildUniqueId)
+		if (ServerBuildId == BuildUniqueId)
 		{
 			// Create a new pending search result 
 			FPendingSearchResultSteam* NewPendingSearch = new (PendingSearchResults) FPendingSearchResultSteam(this);
