Adding CreateSocket that takes protocol types to the base classes along with some other helpers.

diff --git a/Engine/Plugins/Online/OnlineSubsystemSteam/Source/Private/SocketSubsystemSteam.cpp b/Engine/Plugins/Online/OnlineSubsystemSteam/Source/Private/SocketSubsystemSteam.cpp
--- a/Engine/Plugins/Online/OnlineSubsystemSteam/Source/Private/SocketSubsystemSteam.cpp
+++ b/Engine/Plugins/Online/OnlineSubsystemSteam/Source/Private/SocketSubsystemSteam.cpp
@@ -150,11 +150,12 @@ void FSocketSubsystemSteam::Shutdown()
  *
  * @Param SocketType type of socket to create (DGram, Stream, etc)
  * @param SocketDescription debug description
+ * @param ProtocolType the socket protocol to be used
  * @param bForceUDP overrides any platform specific protocol with UDP instead
  *
  * @return the new socket or NULL if failed
  */
-FSocket* FSocketSubsystemSteam::CreateSocket(const FName& SocketType, const FString& SocketDescription, bool bForceUDP)
+FSocket* FSocketSubsystemSteam::CreateSocket(const FName& SocketType, const FString& SocketDescription, ESocketProtocolFamily ProtocolType, bool bForceUDP)
 {
 	FSocket* NewSocket = NULL;
 	if (SocketType == FName("SteamClientSocket"))
@@ -199,10 +200,10 @@ FSocket* FSocketSubsystemSteam::CreateSocket(const FName& SocketType, const FStr
 		ISocketSubsystem* PlatformSocketSub = ISocketSubsystem::Get(PLATFORM_SOCKETSUBSYSTEM);
 		if (PlatformSocketSub)
 		{
-			NewSocket = PlatformSocketSub->CreateSocket(SocketType, SocketDescription, bForceUDP);
+			NewSocket = PlatformSocketSub->CreateSocket(SocketType, SocketDescription, ProtocolType, bForceUDP);
 		}
 	}
-	
+
 	if (!NewSocket)
 	{
 		UE_LOG(LogSockets, Warning, TEXT("Failed to create socket %s [%s]"), *SocketType.ToString(), *SocketDescription);
diff --git a/Engine/Plugins/Online/OnlineSubsystemSteam/Source/Private/SocketSubsystemSteam.h b/Engine/Plugins/Online/OnlineSubsystemSteam/Source/Private/SocketSubsystemSteam.h
--- a/Engine/Plugins/Online/OnlineSubsystemSteam/Source/Private/SocketSubsystemSteam.h
+++ b/Engine/Plugins/Online/OnlineSubsystemSteam/Source/Private/SocketSubsystemSteam.h
@@ -201,12 +201,13 @@ public:
 	 * Creates a socket
 	 *
 	 * @Param SocketType type of socket to create (DGram, Stream, etc)
- 	 * @param SocketDescription debug description
+	 * @param SocketDescription debug description
+	 * @param ProtocolType the socket protocol to be used
 	 * @param bForceUDP overrides any platform specific protocol with UDP instead
 	 *
 	 * @return the new socket or NULL if failed
 	 */
-	virtual class FSocket* CreateSocket(const FName& SocketType, const FString& SocketDescription, bool bForceUDP = false) override;
+	virtual class FSocket* CreateSocket(const FName& SocketType, const FString& SocketDescription, ESocketProtocolFamily ProtocolType, bool bForceUDP = false) override;
 
 	/**
 	 * Cleans up a socket class
diff --git a/Engine/Source/Runtime/Sockets/Private/BSDIPv6Sockets/SocketSubsystemBSDIPv6.cpp b/Engine/Source/Runtime/Sockets/Private/BSDIPv6Sockets/SocketSubsystemBSDIPv6.cpp
--- a/Engine/Source/Runtime/Sockets/Private/BSDIPv6Sockets/SocketSubsystemBSDIPv6.cpp
+++ b/Engine/Source/Runtime/Sockets/Private/BSDIPv6Sockets/SocketSubsystemBSDIPv6.cpp
@@ -46,7 +46,7 @@ ESocketErrors FSocketSubsystemBSDIPv6::TranslateGAIErrorCode(int32 Code) const
 	return SE_NO_ERROR;
 }
 
-FSocket* FSocketSubsystemBSDIPv6::CreateSocket(const FName& SocketType, const FString& SocketDescription, bool bForceUDP)
+FSocket* FSocketSubsystemBSDIPv6::CreateSocket(const FName& SocketType, const FString& SocketDescription, ESocketProtocolFamily ProtocolType, bool bForceUDP)
 {
 	SOCKET Socket = INVALID_SOCKET;
 	FSocket* NewSocket = NULL;
diff --git a/Engine/Source/Runtime/Sockets/Private/BSDIPv6Sockets/SocketSubsystemBSDIPv6.h b/Engine/Source/Runtime/Sockets/Private/BSDIPv6Sockets/SocketSubsystemBSDIPv6.h
--- a/Engine/Source/Runtime/Sockets/Private/BSDIPv6Sockets/SocketSubsystemBSDIPv6.h
+++ b/Engine/Source/Runtime/Sockets/Private/BSDIPv6Sockets/SocketSubsystemBSDIPv6.h
@@ -23,7 +23,7 @@ public:
 
 	virtual TSharedRef<FInternetAddr> CreateInternetAddr( uint32 Address = 0, uint32 Port = 0 ) override;
 
-	virtual class FSocket* CreateSocket( const FName& SocketType, const FString& SocketDescription, bool bForceUDP = false ) override;
+	virtual class FSocket* CreateSocket(const FName& SocketType, const FString& SocketDescription, ESocketProtocolFamily ProtocolType, bool bForceUDP = false) override;
 
 	virtual FResolveInfoCached* CreateResolveInfoCached(TSharedPtr<FInternetAddr> Addr) const override;
 
diff --git a/Engine/Source/Runtime/Sockets/Private/BSDSockets/SocketSubsystemBSD.cpp b/Engine/Source/Runtime/Sockets/Private/BSDSockets/SocketSubsystemBSD.cpp
--- a/Engine/Source/Runtime/Sockets/Private/BSDSockets/SocketSubsystemBSD.cpp
+++ b/Engine/Source/Runtime/Sockets/Private/BSDSockets/SocketSubsystemBSD.cpp
@@ -26,6 +26,14 @@ FSocket* FSocketSubsystemBSD::CreateSocket(const FName& SocketType, const FStrin
 	FSocket* NewSocket = nullptr;
 	int PlatformSpecificTypeFlags = 0;
 
+	// For platforms that have two subsystems (ex: Steam) but don't explicitly inherit from SocketSubsystemBSD
+	// so they don't know which protocol to end up using and pass in None.
+	// This is invalid, so we need to attempt to still resolve it.
+	if (ProtocolType == ESocketProtocolFamily::None)
+	{
+		ProtocolType = GetDefaultSocketProtocolFamily();
+	}
+
 	// Don't support any other protocol families.
 	if (ProtocolType != ESocketProtocolFamily::IPv4 && ProtocolType != ESocketProtocolFamily::IPv6)
 	{
diff --git a/Engine/Source/Runtime/Sockets/Private/IOS/SocketSubsystemIOS.cpp b/Engine/Source/Runtime/Sockets/Private/IOS/SocketSubsystemIOS.cpp
--- a/Engine/Source/Runtime/Sockets/Private/IOS/SocketSubsystemIOS.cpp
+++ b/Engine/Source/Runtime/Sockets/Private/IOS/SocketSubsystemIOS.cpp
@@ -77,9 +77,9 @@ bool FSocketSubsystemIOS::HasNetworkDevice()
 	return true;
 }
 
-FSocket* FSocketSubsystemIOS::CreateSocket(const FName& SocketType, const FString& SocketDescription, bool bForceUDP)
+FSocket* FSocketSubsystemIOS::CreateSocket(const FName& SocketType, const FString& SocketDescription, ESocketProtocolFamily ProtocolType, bool bForceUDP)
 {
-	FSocketBSD* NewSocket = (FSocketBSD*)FSocketSubsystemBSD::CreateSocket(SocketType, SocketDescription, bForceUDP);
+	FSocketBSD* NewSocket = (FSocketBSD*)FSocketSubsystemBSD::CreateSocket(SocketType, SocketDescription, ProtocolType, bForceUDP);
 	if (NewSocket)
 	{
 		NewSocket->SetIPv6Only(false);
diff --git a/Engine/Source/Runtime/Sockets/Private/IOS/SocketSubsystemIOS.h b/Engine/Source/Runtime/Sockets/Private/IOS/SocketSubsystemIOS.h
--- a/Engine/Source/Runtime/Sockets/Private/IOS/SocketSubsystemIOS.h
+++ b/Engine/Source/Runtime/Sockets/Private/IOS/SocketSubsystemIOS.h
@@ -54,7 +54,7 @@ public:
 	virtual bool Init(FString& Error) override;
 	virtual void Shutdown() override;
 	virtual bool HasNetworkDevice() override;
-	virtual FSocket* CreateSocket(const FName& SocketType, const FString& SocketDescription, bool bForceUDP) override;
+	virtual FSocket* CreateSocket(const FName& SocketType, const FString& SocketDescription, ESocketProtocolFamily ProtocolType, bool bForceUDP) override;
 
 	virtual TSharedRef<FInternetAddr> GetLocalHostAddr(FOutputDevice& Out, bool& bCanBindAll) override;
 	virtual class FSocketBSD* InternalBSDSocketFactory(SOCKET Socket, ESocketType SocketType, const FString& SocketDescription) override;
diff --git a/Engine/Source/Runtime/Sockets/Private/Unix/SocketSubsystemUnix.cpp b/Engine/Source/Runtime/Sockets/Private/Unix/SocketSubsystemUnix.cpp
--- a/Engine/Source/Runtime/Sockets/Private/Unix/SocketSubsystemUnix.cpp
+++ b/Engine/Source/Runtime/Sockets/Private/Unix/SocketSubsystemUnix.cpp
@@ -3,7 +3,6 @@
 #include "Unix/SocketSubsystemUnix.h"
 #include "Misc/CommandLine.h"
 #include "SocketSubsystemModule.h"
-#include "IPAddress.h"
 #include "BSDSockets/IPAddressBSD.h"
 
 #include <ifaddrs.h>
@@ -91,6 +90,22 @@ bool FSocketSubsystemUnix::HasNetworkDevice()
 	return true;
 }
 
+FSocket* FSocketSubsystemUnix::CreateSocket(const FName& SocketType, const FString& SocketDescription, ESocketProtocolFamily ProtocolType, bool bForceUDP)
+{
+	FSocketBSD* NewSocket = (FSocketBSD*)FSocketSubsystemBSD::CreateSocket(SocketType, SocketDescription, ProtocolType, bForceUDP);
+
+	if (NewSocket != nullptr)
+	{
+		NewSocket->SetIPv6Only(false);
+	}
+	else
+	{
+		UE_LOG(LogSockets, Warning, TEXT("Failed to create socket %s [%s]"), *SocketType.ToString(), *SocketDescription);
+	}
+
+	return NewSocket;
+}
+
 TSharedRef<FInternetAddr> FSocketSubsystemUnix::GetLocalHostAddr(FOutputDevice& Out, bool& bCanBindAll)
 {
 	// get parent address first
diff --git a/Engine/Source/Runtime/Sockets/Private/Unix/SocketSubsystemUnix.h b/Engine/Source/Runtime/Sockets/Private/Unix/SocketSubsystemUnix.h
--- a/Engine/Source/Runtime/Sockets/Private/Unix/SocketSubsystemUnix.h
+++ b/Engine/Source/Runtime/Sockets/Private/Unix/SocketSubsystemUnix.h
@@ -50,6 +50,7 @@ public:
 	virtual bool Init(FString& Error) override;
 	virtual void Shutdown() override;
 	virtual bool HasNetworkDevice() override;
+	virtual class FSocket* CreateSocket(const FName& SocketType, const FString& SocketDescription, ESocketProtocolFamily ProtocolType, bool bForceUDP = false) override;
 	virtual TSharedRef<FInternetAddr> GetLocalHostAddr(FOutputDevice& Out, bool& bCanBindAll) override;
 	virtual bool GetLocalAdapterAddresses(TArray<TSharedPtr<FInternetAddr> >& OutAddresses) override;
 };
diff --git a/Engine/Source/Runtime/Sockets/Private/Windows/SocketSubsystemWindows.cpp b/Engine/Source/Runtime/Sockets/Private/Windows/SocketSubsystemWindows.cpp
--- a/Engine/Source/Runtime/Sockets/Private/Windows/SocketSubsystemWindows.cpp
+++ b/Engine/Source/Runtime/Sockets/Private/Windows/SocketSubsystemWindows.cpp
@@ -69,9 +69,9 @@ void FSocketSubsystemWindows::Destroy()
 /* FSocketSubsystemBSD overrides
 *****************************************************************************/
 
-FSocket* FSocketSubsystemWindows::CreateSocket(const FName& SocketType, const FString& SocketDescription, bool bForceUDP)
+FSocket* FSocketSubsystemWindows::CreateSocket(const FName& SocketType, const FString& SocketDescription, ESocketProtocolFamily ProtocolType, bool bForceUDP)
 {
-	FSocketBSD* NewSocket = (FSocketBSD*)FSocketSubsystemBSD::CreateSocket(SocketType, SocketDescription, bForceUDP);
+	FSocketBSD* NewSocket = (FSocketBSD*)FSocketSubsystemBSD::CreateSocket(SocketType, SocketDescription, ProtocolType, bForceUDP);
 
 	if (NewSocket != nullptr)
 	{
diff --git a/Engine/Source/Runtime/Sockets/Private/Windows/SocketSubsystemWindows.h b/Engine/Source/Runtime/Sockets/Private/Windows/SocketSubsystemWindows.h
--- a/Engine/Source/Runtime/Sockets/Private/Windows/SocketSubsystemWindows.h
+++ b/Engine/Source/Runtime/Sockets/Private/Windows/SocketSubsystemWindows.h
@@ -28,7 +28,7 @@ public:
 
 	// FSocketSubsystemBSD overrides
 
-	virtual class FSocket* CreateSocket(const FName& SocketType, const FString& SocketDescription, bool bForceUDP = false) override;
+	virtual class FSocket* CreateSocket(const FName& SocketType, const FString& SocketDescription, ESocketProtocolFamily ProtocolType, bool bForceUDP = false) override;
 	virtual bool HasNetworkDevice() override;
 	virtual ESocketErrors GetLastErrorCode() override;
 	virtual bool GetLocalAdapterAddresses( TArray<TSharedPtr<FInternetAddr> >& OutAdresses ) override;
diff --git a/Engine/Source/Runtime/Sockets/Public/SocketSubsystem.h b/Engine/Source/Runtime/Sockets/Public/SocketSubsystem.h
--- a/Engine/Source/Runtime/Sockets/Public/SocketSubsystem.h
+++ b/Engine/Source/Runtime/Sockets/Public/SocketSubsystem.h
@@ -82,7 +82,22 @@ public:
 	 *
 	 * @return the new socket or NULL if failed
 	 */
-	virtual class FSocket* CreateSocket(const FName& SocketType, const FString& SocketDescription, bool bForceUDP = false) = 0;
+	virtual class FSocket* CreateSocket(const FName& SocketType, const FString& SocketDescription, bool bForceUDP = false)
+	{
+		return CreateSocket(SocketType, SocketDescription, ESocketProtocolFamily::None, bForceUDP);
+	}
+
+	/**
+	 * Creates a socket
+	 *
+	 * @Param SocketType type of socket to create (DGram, Stream, etc)
+	 * @param SocketDescription debug description
+	 * @param ProtocolType the socket protocol to be used
+	 * @param bForceUDP overrides any platform specific protocol with UDP instead
+	 *
+	 * @return the new socket or NULL if failed
+	 */
+	virtual class FSocket* CreateSocket(const FName& SocketType, const FString& SocketDescription, ESocketProtocolFamily ProtocolType, bool bForceUDP = false) = 0;
 
 	/**
 	 * Creates a resolve info cached struct to hold the resolved address
