Add support for the various types of friends lists that Unreal supports on Steam. Fixing up detection of user statuses.

diff --git a/Engine/Plugins/Online/OnlineSubsystemSteam/Source/Private/OnlineFriendsInterfaceSteam.cpp b/Engine/Plugins/Online/OnlineSubsystemSteam/Source/Private/OnlineFriendsInterfaceSteam.cpp
--- a/Engine/Plugins/Online/OnlineSubsystemSteam/Source/Private/OnlineFriendsInterfaceSteam.cpp
+++ b/Engine/Plugins/Online/OnlineSubsystemSteam/Source/Private/OnlineFriendsInterfaceSteam.cpp
@@ -4,8 +4,6 @@
 #include "OnlineSubsystemSteam.h"
 #include "OnlineSubsystemSteamTypes.h"
 
-const FString DefaultFriendsList(EFriendsLists::ToString(EFriendsLists::Default));
-
 // FOnlineFriendSteam
 FOnlineFriendSteam::FOnlineFriendSteam(const CSteamID& InUserId)
 	: UserId(new FUniqueNetIdSteam(InUserId))
@@ -61,16 +59,12 @@ FOnlineFriendsSteam::FOnlineFriendsSteam(FOnlineSubsystemSteam* InSteamSubsystem
 bool FOnlineFriendsSteam::ReadFriendsList(int32 LocalUserNum, const FString& ListName, const FOnReadFriendsListComplete& Delegate /*= FOnReadFriendsListComplete()*/)
 {
 	FString ErrorStr;
-	if (!ListName.Equals(DefaultFriendsList, ESearchCase::IgnoreCase))
-	{
-		UE_LOG_ONLINE(Warning, TEXT("Only the default friends list is supported. ListName=%s"), *ListName);
-	}
 	if (LocalUserNum < MAX_LOCAL_PLAYERS &&
 		SteamUserPtr != NULL &&
 		SteamUserPtr->BLoggedOn() &&
 		SteamFriendsPtr != NULL)
 	{
-		SteamSubsystem->QueueAsyncTask(new FOnlineAsyncTaskSteamReadFriendsList(this, LocalUserNum, Delegate));
+		SteamSubsystem->QueueAsyncTask(new FOnlineAsyncTaskSteamReadFriendsList(this, LocalUserNum, ListName, Delegate));
 	}
 	else
 	{
@@ -212,6 +206,24 @@ void FOnlineFriendsSteam::DumpBlockedPlayers() const
 {
 }
 
+bool FOnlineAsyncTaskSteamReadFriendsList::CanAddUserToList(bool bIsOnline, bool bIsPlayingThisGame, bool bIsPlayingGameInSession)
+{
+	switch (FriendsListFilter)
+	{
+		default:
+		case EFriendsLists::Default:
+			return true;
+		case EFriendsLists::OnlinePlayers:
+			return bIsOnline;
+		case EFriendsLists::InGamePlayers:
+			return bIsOnline && bIsPlayingThisGame;
+		case EFriendsLists::InGameAndSessionPlayers:
+			return bIsOnline && bIsPlayingThisGame && bIsPlayingGameInSession;
+	}
+
+	return false;
+}
+
 void FOnlineAsyncTaskSteamReadFriendsList::Finalize()
 {
 	FOnlineSubsystemSteam* SteamSubsystem = FriendsPtr->SteamSubsystem;
@@ -226,22 +238,42 @@ void FOnlineAsyncTaskSteamReadFriendsList::Finalize()
 	{
 		const CSteamID SteamPlayerId = SteamFriendsPtr->GetFriendByIndex(Index, k_EFriendFlagImmediate);
 		const FString NickName(UTF8_TO_TCHAR(SteamFriendsPtr->GetFriendPersonaName(SteamPlayerId)));
-		// Non-unique named friends are skipped. Don't want impersonation for banning, etc.
-		if (NickName.Len() > 0)
+
+		// Get this user's friend information so we can figure out if we can add them to our list.
+		bool bInASession;
+		FriendGameInfo_t FriendGameInfo;
+		bool bIsPlayingAGame = SteamFriendsPtr->GetFriendGamePlayed(SteamPlayerId, &FriendGameInfo);
+		bool bIsOnline = (SteamFriendsPtr->GetFriendPersonaState(SteamPlayerId) >= k_EPersonaStateOnline);
+		bool bIsPlayingThisGame = (FriendGameInfo.m_gameID.AppID() == SteamSubsystem->GetSteamAppId());
+		bool bHasConnectInformation = (SteamFriendsPtr->GetFriendRichPresence(SteamPlayerId, "connect") != nullptr);
+		FString JoinablePresenceString = UTF8_TO_TCHAR(SteamFriendsPtr->GetFriendRichPresence(SteamPlayerId, "Joinable"));
+
+		// Platforms can override joinability using the "Joinable", which overrides the default check
+		// Remote friend is responsible for updating their presence to have the joinable status
+		if (!JoinablePresenceString.IsEmpty())
+		{
+			bInASession = (JoinablePresenceString == TEXT("true"));
+		}
+		else
+		{
+			bInASession = bIsPlayingThisGame && bHasConnectInformation;
+		}
+
+		// Skip invalid entries and ones that do not fit our current filters.
+		if (NickName.Len() > 0 && CanAddUserToList(bIsOnline, bIsPlayingThisGame, bInASession))
 		{
 			// Add to list
 			TSharedRef<FOnlineFriendSteam> Friend(new FOnlineFriendSteam(SteamPlayerId));
 			FriendsList.Friends.Add(Friend);
-			// Get their game playing information
-			FriendGameInfo_t FriendGameInfo;
-			const bool bIsPlayingAGame = SteamFriendsPtr->GetFriendGamePlayed(SteamPlayerId, &FriendGameInfo);
+
 			// Now fill in the friend info
 			Friend->AccountData.Add(TEXT("nickname"), NickName);
 			Friend->Presence.Status.StatusStr = UTF8_TO_TCHAR(SteamFriendsPtr->GetFriendRichPresence(SteamPlayerId,"status"));
-			FString JoinablePresenceString = UTF8_TO_TCHAR(SteamFriendsPtr->GetFriendRichPresence(SteamPlayerId,"Joinable"));
-			// Remote friend is responsible for updating their presence to have the joinable status
-			Friend->Presence.bIsJoinable = JoinablePresenceString == TEXT("true");
-			Friend->Presence.bIsOnline = SteamFriendsPtr->GetFriendPersonaState(SteamPlayerId) > k_EPersonaStateOffline;
+			Friend->Presence.bIsJoinable = bInASession;
+			Friend->Presence.bIsOnline = bIsOnline;
+			Friend->Presence.bIsPlaying = bIsPlayingAGame;
+			Friend->Presence.bIsPlayingThisGame = bIsPlayingThisGame;
+
 			switch (SteamFriendsPtr->GetFriendPersonaState(SteamPlayerId))
 			{
 				case k_EPersonaStateOffline:
@@ -260,9 +292,6 @@ void FOnlineAsyncTaskSteamReadFriendsList::Finalize()
 					Friend->Presence.Status.State = EOnlinePresenceState::Online;
 					break;
 			}
-			Friend->Presence.bIsPlaying = bIsPlayingAGame;
-			// Check the steam id
-			Friend->Presence.bIsPlayingThisGame = FriendGameInfo.m_gameID.AppID() == SteamSubsystem->GetSteamAppId();
 			// Remote friend is responsible for updating their presence to have the voice flag
 			FString VoicePresenceString = UTF8_TO_TCHAR(SteamFriendsPtr->GetFriendRichPresence(SteamPlayerId,"HasVoice"));
 			// Determine if the user has voice support
@@ -275,5 +304,5 @@ void FOnlineAsyncTaskSteamReadFriendsList::TriggerDelegates(void)
 {
 	FOnlineAsyncTask::TriggerDelegates();
 
-	Delegate.ExecuteIfBound(LocalUserNum, true, DefaultFriendsList, FString());
+	Delegate.ExecuteIfBound(LocalUserNum, true, EFriendsLists::ToString(FriendsListFilter), FString());
 }
diff --git a/Engine/Plugins/Online/OnlineSubsystemSteam/Source/Private/OnlineFriendsInterfaceSteam.h b/Engine/Plugins/Online/OnlineSubsystemSteam/Source/Private/OnlineFriendsInterfaceSteam.h
--- a/Engine/Plugins/Online/OnlineSubsystemSteam/Source/Private/OnlineFriendsInterfaceSteam.h
+++ b/Engine/Plugins/Online/OnlineSubsystemSteam/Source/Private/OnlineFriendsInterfaceSteam.h
@@ -139,9 +139,13 @@ class FOnlineAsyncTaskSteamReadFriendsList :
 	FOnlineFriendsSteam* FriendsPtr;
 	/** The user that is triggering the event */
 	int32 LocalUserNum;
+	/** The type of list we are trying to fetch */
+	EFriendsLists::Type FriendsListFilter;
 
 	FOnReadFriendsListComplete Delegate;
 
+	bool CanAddUserToList(bool bIsOnline, bool bIsPlayingThisGame, bool bIsPlayingGameInSession);
+
 public:
 	/**
 	 * Inits the pointer used to trigger the delegates on
@@ -150,12 +154,29 @@ public:
 	 * @param InLocalUserNum the local user that requested the read
 	 * @param InDelegate the delegate that will be called when reading the friends list is complete
 	 */
-	FOnlineAsyncTaskSteamReadFriendsList(FOnlineFriendsSteam* InFriendsPtr, int32 InLocalUserNum, const FOnReadFriendsListComplete& InDelegate) :
+	FOnlineAsyncTaskSteamReadFriendsList(FOnlineFriendsSteam* InFriendsPtr, int32 InLocalUserNum, const FString& InFriendsListFilter, const FOnReadFriendsListComplete& InDelegate) :
 		FriendsPtr(InFriendsPtr),
 		LocalUserNum(InLocalUserNum),
 		Delegate(InDelegate)
 	{
 		check(FriendsPtr);
+
+		if (InFriendsListFilter.Equals(EFriendsLists::ToString(EFriendsLists::OnlinePlayers), ESearchCase::IgnoreCase))
+		{
+			FriendsListFilter = EFriendsLists::OnlinePlayers;
+		}
+		else if (InFriendsListFilter.Equals(EFriendsLists::ToString(EFriendsLists::InGamePlayers), ESearchCase::IgnoreCase))
+		{
+			FriendsListFilter = EFriendsLists::InGamePlayers;
+		}
+		else if (InFriendsListFilter.Equals(EFriendsLists::ToString(EFriendsLists::InGameAndSessionPlayers), ESearchCase::IgnoreCase))
+		{
+			FriendsListFilter = EFriendsLists::InGameAndSessionPlayers;
+		}
+		else
+		{
+			FriendsListFilter = EFriendsLists::Default;
+		}
 	}
 
 	// FOnlineAsyncTask
