PR #4976: Remove redundant logic in OSSSteam (Contributed by [REDACTED_PERSON])

#github https://github.com/EpicGames/UnrealEngine/pull/4976

diff --git a/Engine/Plugins/Online/OnlineSubsystemSteam/Source/Private/OnlineSessionAsyncServerSteam.cpp b/Engine/Plugins/Online/OnlineSubsystemSteam/Source/Private/OnlineSessionAsyncServerSteam.cpp
--- a/Engine/Plugins/Online/OnlineSubsystemSteam/Source/Private/OnlineSessionAsyncServerSteam.cpp
+++ b/Engine/Plugins/Online/OnlineSubsystemSteam/Source/Private/OnlineSessionAsyncServerSteam.cpp
@@ -668,11 +668,7 @@ bool FPendingSearchResultSteam::FillSessionFromServerRules()
 	if (bSuccess && (KeysFound == STEAMKEY_NUMREQUIREDSERVERKEYS) && (SteamAddrKeysFound == 2))
 	{
 		SessionInfo->HostAddr = HostAddr;
-
-		if (SteamAddrKeysFound == 2) //-V547
-		{
-			SessionInfo->SteamP2PAddr = SteamP2PAddr;
-		}
+		SessionInfo->SteamP2PAddr = SteamP2PAddr;
 
 		Session->SessionInfo = SessionInfo;
 		return true;
