Fix bad user lookup due to incorrect base pointers in Steam, causing the server to crash

diff --git a/Engine/Plugins/Online/OnlineSubsystemSteam/Source/Private/OnlineSessionInterfaceSteam.cpp b/Engine/Plugins/Online/OnlineSubsystemSteam/Source/Private/OnlineSessionInterfaceSteam.cpp
--- a/Engine/Plugins/Online/OnlineSubsystemSteam/Source/Private/OnlineSessionInterfaceSteam.cpp
+++ b/Engine/Plugins/Online/OnlineSubsystemSteam/Source/Private/OnlineSessionInterfaceSteam.cpp
@@ -228,8 +228,19 @@ bool FOnlineSessionSteam::CreateSession(int32 HostingPlayerNum, FName SessionNam
 		Session->NumOpenPublicConnections = NewSessionSettings.bIsDedicated ? NewSessionSettings.NumPublicConnections : NewSessionSettings.NumPublicConnections - 1;
 
 		Session->HostingPlayerNum = HostingPlayerNum;
-		Session->OwningUserId = SteamUser() ? MakeShareable(new FUniqueNetIdSteam(SteamUser()->GetSteamID())) : NULL;
 		Session->OwningUserName = SteamFriends() ? SteamFriends()->GetPersonaName() : FString(TEXT(""));
+		if (SteamUser() != nullptr)
+		{
+			Session->OwningUserId = MakeShareable(new FUniqueNetIdSteam(SteamUser()->GetSteamID()));
+		}
+		else if (SteamGameServer() != nullptr)
+		{
+			Session->OwningUserId = MakeShareable(new FUniqueNetIdSteam(SteamGameServer()->GetSteamID()));
+		}
+		else
+		{
+			Session->OwningUserId = nullptr;
+		}
 		
 		// Unique identifier of this build for compatibility
 		Session->SessionSettings.BuildUniqueId = GetBuildUniqueId();
