Fix Steam SocketSubsystem GetAddressInfo filtering such that it determines input strings better.

#rnx

diff --git a/Engine/Plugins/Online/OnlineSubsystemSteam/Source/Private/SocketSubsystemSteam.cpp b/Engine/Plugins/Online/OnlineSubsystemSteam/Source/Private/SocketSubsystemSteam.cpp
--- a/Engine/Plugins/Online/OnlineSubsystemSteam/Source/Private/SocketSubsystemSteam.cpp
+++ b/Engine/Plugins/Online/OnlineSubsystemSteam/Source/Private/SocketSubsystemSteam.cpp
@@ -341,12 +341,13 @@ void FSocketSubsystemSteam::ConnectFailure(const FUniqueNetIdSteam& RemoteId)
 FAddressInfoResult FSocketSubsystemSteam::GetAddressInfo(const TCHAR* HostName, const TCHAR* ServiceName,
 	EAddressInfoFlags QueryFlags, const FName ProtocolTypeName, ESocketType SocketType)
 {
-	int32 UnusedIndex;
 	FString RawAddress(HostName);
+
 	// Remove steam prefixes if they exist
 	RawAddress.RemoveFromStart(STEAM_URL_PREFIX);
 
-	if (!RawAddress.FindChar(TEXT(':'), UnusedIndex) && !RawAddress.FindChar(TEXT('.'), UnusedIndex))
+	// Steam ids are pure numeric values, so we can use this to determine if the input is a SteamID.
+	if (RawAddress.IsNumeric())
 	{
 		FAddressInfoResult SteamResult(HostName, ServiceName);
 
@@ -393,11 +394,11 @@ FAddressInfoResult FSocketSubsystemSteam::GetAddressInfo(const TCHAR* HostName,
  */
 TSharedPtr<FInternetAddr> FSocketSubsystemSteam::GetAddressFromString(const FString& InAddress)
 {
-	int32 UnusedIndex;
 	FString RawAddress = InAddress;
 	// Remove steam prefixes if they exist
 	RawAddress.RemoveFromStart(STEAM_URL_PREFIX);
-	if (!RawAddress.FindChar(TEXT(':'), UnusedIndex) && !RawAddress.FindChar(TEXT('.'), UnusedIndex))
+
+	if (RawAddress.IsNumeric())
 	{
 		// This is a Steam Address
 		uint64 Id = FCString::Atoi64(*RawAddress);
