Fix invalid query issue when searching for steam dedicated servers that friends are in due to incorrect address types being created.

#rnx

diff --git a/Engine/Plugins/Online/OnlineSubsystemSteam/Source/Private/OnlineSessionInterfaceSteam.cpp b/Engine/Plugins/Online/OnlineSubsystemSteam/Source/Private/OnlineSessionInterfaceSteam.cpp
--- a/Engine/Plugins/Online/OnlineSubsystemSteam/Source/Private/OnlineSessionInterfaceSteam.cpp
+++ b/Engine/Plugins/Online/OnlineSubsystemSteam/Source/Private/OnlineSessionInterfaceSteam.cpp
@@ -1042,7 +1042,7 @@ bool FOnlineSessionSteam::FindFriendSession(int32 LocalUserNum, const FUniqueNet
 				else
 				{
 					// Search for the session via host ip
-					TSharedRef<FInternetAddr> IpAddr = ISocketSubsystem::Get()->CreateInternetAddr();
+					TSharedRef<FInternetAddr> IpAddr = ISocketSubsystem::Get(PLATFORM_SOCKETSUBSYSTEM)->CreateInternetAddr();
 					IpAddr->SetIp(FriendGameInfo.m_unGameIP);
 					IpAddr->SetPort(FriendGameInfo.m_usGamePort);
 					CurrentSessionSearch->QuerySettings.Set(FName(SEARCH_STEAM_HOSTIP), IpAddr->ToString(true), EOnlineComparisonOp::Equals);
