Skip lobbies we are already a part of when generating matchmaking results. This allows us to not try to fetch data on ourselves.

diff --git a/Engine/Plugins/Online/OnlineSubsystemSteam/Source/Private/OnlineSessionAsyncLobbySteam.cpp b/Engine/Plugins/Online/OnlineSubsystemSteam/Source/Private/OnlineSessionAsyncLobbySteam.cpp
--- a/Engine/Plugins/Online/OnlineSubsystemSteam/Source/Private/OnlineSessionAsyncLobbySteam.cpp
+++ b/Engine/Plugins/Online/OnlineSubsystemSteam/Source/Private/OnlineSessionAsyncLobbySteam.cpp
@@ -950,11 +950,19 @@ void FOnlineAsyncTaskSteamFindLobbiesBase::Tick()
 				bWasSuccessful = bSuccessCallResult && !bFailedCall && !bFailedResult;
 				if (bWasSuccessful)
 				{
+					FOnlineSessionSteamPtr SessionInt = StaticCastSharedPtr<FOnlineSessionSteam>(Subsystem->GetSessionInterface());
+					check(SessionInt.IsValid());
+
 					// Trigger the lobby data requests
 					int32 NumLobbies = (int32)CallbackResults.m_nLobbiesMatching;
 					for (int32 LobbyIdx = 0; LobbyIdx < NumLobbies; LobbyIdx++)
 					{
-						LobbyIDs.Add(SteamMatchmakingPtr->GetLobbyByIndex(LobbyIdx));
+						// We cannot matchmake any lobbies that we are currently in as this is unsupported by the platform
+						FUniqueNetIdSteam LobbyID(SteamMatchmakingPtr->GetLobbyByIndex(LobbyIdx));
+						if (!SessionInt->IsMemberOfLobby(LobbyID))
+						{
+							LobbyIDs.Add(LobbyID);
+						}
 					}
 					FindLobbiesState = EFindLobbiesState::RequestLobbyData;
 				}
