Remove the bForceUDP from the new CreateSocket signature. This is an argument that came in from UE3 and is not supported. It should not continue to exist anymore.

diff --git a/Engine/Plugins/Online/OnlineSubsystemSteam/Source/Private/SocketSubsystemSteam.cpp b/Engine/Plugins/Online/OnlineSubsystemSteam/Source/Private/SocketSubsystemSteam.cpp
--- a/Engine/Plugins/Online/OnlineSubsystemSteam/Source/Private/SocketSubsystemSteam.cpp
+++ b/Engine/Plugins/Online/OnlineSubsystemSteam/Source/Private/SocketSubsystemSteam.cpp
@@ -151,11 +151,10 @@ void FSocketSubsystemSteam::Shutdown()
  * @Param SocketType type of socket to create (DGram, Stream, etc)
  * @param SocketDescription debug description
  * @param ProtocolType the socket protocol to be used
- * @param bForceUDP overrides any platform specific protocol with UDP instead
  *
  * @return the new socket or NULL if failed
  */
-FSocket* FSocketSubsystemSteam::CreateSocket(const FName& SocketType, const FString& SocketDescription, ESocketProtocolFamily ProtocolType, bool bForceUDP)
+FSocket* FSocketSubsystemSteam::CreateSocket(const FName& SocketType, const FString& SocketDescription, ESocketProtocolFamily ProtocolType)
 {
 	FSocket* NewSocket = NULL;
 	if (SocketType == FName("SteamClientSocket"))
@@ -200,7 +199,7 @@ FSocket* FSocketSubsystemSteam::CreateSocket(const FName& SocketType, const FStr
 		ISocketSubsystem* PlatformSocketSub = ISocketSubsystem::Get(PLATFORM_SOCKETSUBSYSTEM);
 		if (PlatformSocketSub)
 		{
-			NewSocket = PlatformSocketSub->CreateSocket(SocketType, SocketDescription, ProtocolType, bForceUDP);
+			NewSocket = PlatformSocketSub->CreateSocket(SocketType, SocketDescription, ProtocolType);
 		}
 	}
 
diff --git a/Engine/Plugins/Online/OnlineSubsystemSteam/Source/Private/SocketSubsystemSteam.h b/Engine/Plugins/Online/OnlineSubsystemSteam/Source/Private/SocketSubsystemSteam.h
--- a/Engine/Plugins/Online/OnlineSubsystemSteam/Source/Private/SocketSubsystemSteam.h
+++ b/Engine/Plugins/Online/OnlineSubsystemSteam/Source/Private/SocketSubsystemSteam.h
@@ -203,11 +203,10 @@ public:
 	 * @Param SocketType type of socket to create (DGram, Stream, etc)
 	 * @param SocketDescription debug description
 	 * @param ProtocolType the socket protocol to be used
-	 * @param bForceUDP overrides any platform specific protocol with UDP instead
 	 *
 	 * @return the new socket or NULL if failed
 	 */
-	virtual class FSocket* CreateSocket(const FName& SocketType, const FString& SocketDescription, ESocketProtocolFamily ProtocolType, bool bForceUDP = false) override;
+	virtual class FSocket* CreateSocket(const FName& SocketType, const FString& SocketDescription, ESocketProtocolFamily ProtocolType) override;
 
 	/**
 	 * Cleans up a socket class
diff --git a/Engine/Source/Runtime/Sockets/Private/BSDIPv6Sockets/SocketSubsystemBSDIPv6.cpp b/Engine/Source/Runtime/Sockets/Private/BSDIPv6Sockets/SocketSubsystemBSDIPv6.cpp
--- a/Engine/Source/Runtime/Sockets/Private/BSDIPv6Sockets/SocketSubsystemBSDIPv6.cpp
+++ b/Engine/Source/Runtime/Sockets/Private/BSDIPv6Sockets/SocketSubsystemBSDIPv6.cpp
@@ -46,7 +46,7 @@ ESocketErrors FSocketSubsystemBSDIPv6::TranslateGAIErrorCode(int32 Code) const
 	return SE_NO_ERROR;
 }
 
-FSocket* FSocketSubsystemBSDIPv6::CreateSocket(const FName& SocketType, const FString& SocketDescription, ESocketProtocolFamily ProtocolType, bool bForceUDP)
+FSocket* FSocketSubsystemBSDIPv6::CreateSocket(const FName& SocketType, const FString& SocketDescription, ESocketProtocolFamily ProtocolType)
 {
 	SOCKET Socket = INVALID_SOCKET;
 	FSocket* NewSocket = NULL;
diff --git a/Engine/Source/Runtime/Sockets/Private/BSDIPv6Sockets/SocketSubsystemBSDIPv6.h b/Engine/Source/Runtime/Sockets/Private/BSDIPv6Sockets/SocketSubsystemBSDIPv6.h
--- a/Engine/Source/Runtime/Sockets/Private/BSDIPv6Sockets/SocketSubsystemBSDIPv6.h
+++ b/Engine/Source/Runtime/Sockets/Private/BSDIPv6Sockets/SocketSubsystemBSDIPv6.h
@@ -23,7 +23,7 @@ public:
 
 	virtual TSharedRef<FInternetAddr> CreateInternetAddr( uint32 Address = 0, uint32 Port = 0 ) override;
 
-	virtual class FSocket* CreateSocket(const FName& SocketType, const FString& SocketDescription, ESocketProtocolFamily ProtocolType, bool bForceUDP = false) override;
+	virtual class FSocket* CreateSocket(const FName& SocketType, const FString& SocketDescription, ESocketProtocolFamily ProtocolType) override;
 
 	virtual FResolveInfoCached* CreateResolveInfoCached(TSharedPtr<FInternetAddr> Addr) const override;
 
diff --git a/Engine/Source/Runtime/Sockets/Private/BSDSockets/SocketSubsystemBSD.cpp b/Engine/Source/Runtime/Sockets/Private/BSDSockets/SocketSubsystemBSD.cpp
--- a/Engine/Source/Runtime/Sockets/Private/BSDSockets/SocketSubsystemBSD.cpp
+++ b/Engine/Source/Runtime/Sockets/Private/BSDSockets/SocketSubsystemBSD.cpp
@@ -17,10 +17,10 @@ FSocketBSD* FSocketSubsystemBSD::InternalBSDSocketFactory(SOCKET Socket, ESocket
 
 FSocket* FSocketSubsystemBSD::CreateSocket(const FName& SocketType, const FString& SocketDescription, bool bForceUDP)
 {
-	return CreateSocket(SocketType, SocketDescription, GetDefaultSocketProtocolFamily(), bForceUDP);
+	return CreateSocket(SocketType, SocketDescription, GetDefaultSocketProtocolFamily());
 }
 
-FSocket* FSocketSubsystemBSD::CreateSocket(const FName& SocketType, const FString& SocketDescription, ESocketProtocolFamily ProtocolType, bool bForceUDP)
+FSocket* FSocketSubsystemBSD::CreateSocket(const FName& SocketType, const FString& SocketDescription, ESocketProtocolFamily ProtocolType)
 {
 	SOCKET Socket = INVALID_SOCKET;
 	FSocket* NewSocket = nullptr;
diff --git a/Engine/Source/Runtime/Sockets/Private/BSDSockets/SocketSubsystemBSD.h b/Engine/Source/Runtime/Sockets/Private/BSDSockets/SocketSubsystemBSD.h
--- a/Engine/Source/Runtime/Sockets/Private/BSDSockets/SocketSubsystemBSD.h
+++ b/Engine/Source/Runtime/Sockets/Private/BSDSockets/SocketSubsystemBSD.h
@@ -41,7 +41,7 @@ public:
 	// ISocketSubsystem interface
 	virtual TSharedRef<FInternetAddr> CreateInternetAddr(uint32 Address = 0, uint32 Port = 0) override;
 	virtual class FSocket* CreateSocket( const FName& SocketType, const FString& SocketDescription, bool bForceUDP = false ) override;
-	virtual class FSocket* CreateSocket(const FName& SocketType, const FString& SocketDescription, ESocketProtocolFamily ProtocolType, bool bForceUDP = false);
+	virtual class FSocket* CreateSocket(const FName& SocketType, const FString& SocketDescription, ESocketProtocolFamily ProtocolType) override;
 	virtual void DestroySocket( class FSocket* Socket ) override;
 
 	virtual FAddressInfoResult GetAddressInfo(const TCHAR* HostName, const TCHAR* ServiceName = nullptr,
diff --git a/Engine/Source/Runtime/Sockets/Private/IOS/SocketSubsystemIOS.cpp b/Engine/Source/Runtime/Sockets/Private/IOS/SocketSubsystemIOS.cpp
--- a/Engine/Source/Runtime/Sockets/Private/IOS/SocketSubsystemIOS.cpp
+++ b/Engine/Source/Runtime/Sockets/Private/IOS/SocketSubsystemIOS.cpp
@@ -96,9 +96,9 @@ ESocketErrors FSocketSubsystemIOS::GetHostByName(const ANSICHAR* HostName, FInte
 	return SE_HOST_NOT_FOUND;
 }
 
-FSocket* FSocketSubsystemIOS::CreateSocket(const FName& SocketType, const FString& SocketDescription, ESocketProtocolFamily ProtocolType, bool bForceUDP)
+FSocket* FSocketSubsystemIOS::CreateSocket(const FName& SocketType, const FString& SocketDescription, ESocketProtocolFamily ProtocolType)
 {
-	FSocketBSD* NewSocket = (FSocketBSD*)FSocketSubsystemBSD::CreateSocket(SocketType, SocketDescription, ProtocolType, bForceUDP);
+	FSocketBSD* NewSocket = (FSocketBSD*)FSocketSubsystemBSD::CreateSocket(SocketType, SocketDescription, ProtocolType);
 	if (NewSocket)
 	{
 		NewSocket->SetIPv6Only(false);
diff --git a/Engine/Source/Runtime/Sockets/Private/IOS/SocketSubsystemIOS.h b/Engine/Source/Runtime/Sockets/Private/IOS/SocketSubsystemIOS.h
--- a/Engine/Source/Runtime/Sockets/Private/IOS/SocketSubsystemIOS.h
+++ b/Engine/Source/Runtime/Sockets/Private/IOS/SocketSubsystemIOS.h
@@ -57,7 +57,7 @@ public:
 	virtual bool Init(FString& Error) override;
 	virtual void Shutdown() override;
 	virtual bool HasNetworkDevice() override;
-	virtual FSocket* CreateSocket(const FName& SocketType, const FString& SocketDescription, ESocketProtocolFamily ProtocolType, bool bForceUDP) override;
+	virtual FSocket* CreateSocket(const FName& SocketType, const FString& SocketDescription, ESocketProtocolFamily ProtocolType) override;
 
 	virtual TSharedRef<FInternetAddr> GetLocalHostAddr(FOutputDevice& Out, bool& bCanBindAll) override;
 	virtual class FSocketBSD* InternalBSDSocketFactory(SOCKET Socket, ESocketType SocketType, const FString& SocketDescription) override;
diff --git a/Engine/Source/Runtime/Sockets/Private/Unix/SocketSubsystemUnix.cpp b/Engine/Source/Runtime/Sockets/Private/Unix/SocketSubsystemUnix.cpp
--- a/Engine/Source/Runtime/Sockets/Private/Unix/SocketSubsystemUnix.cpp
+++ b/Engine/Source/Runtime/Sockets/Private/Unix/SocketSubsystemUnix.cpp
@@ -91,9 +91,9 @@ bool FSocketSubsystemUnix::HasNetworkDevice()
 	return true;
 }
 
-FSocket* FSocketSubsystemUnix::CreateSocket(const FName& SocketType, const FString& SocketDescription, ESocketProtocolFamily ProtocolType, bool bForceUDP)
+FSocket* FSocketSubsystemUnix::CreateSocket(const FName& SocketType, const FString& SocketDescription, ESocketProtocolFamily ProtocolType)
 {
-	FSocketBSD* NewSocket = (FSocketBSD*)FSocketSubsystemBSD::CreateSocket(SocketType, SocketDescription, ProtocolType, bForceUDP);
+	FSocketBSD* NewSocket = (FSocketBSD*)FSocketSubsystemBSD::CreateSocket(SocketType, SocketDescription, ProtocolType);
 
 	if (NewSocket != nullptr)
 	{
diff --git a/Engine/Source/Runtime/Sockets/Private/Unix/SocketSubsystemUnix.h b/Engine/Source/Runtime/Sockets/Private/Unix/SocketSubsystemUnix.h
--- a/Engine/Source/Runtime/Sockets/Private/Unix/SocketSubsystemUnix.h
+++ b/Engine/Source/Runtime/Sockets/Private/Unix/SocketSubsystemUnix.h
@@ -50,7 +50,7 @@ public:
 	virtual bool Init(FString& Error) override;
 	virtual void Shutdown() override;
 	virtual bool HasNetworkDevice() override;
-	virtual class FSocket* CreateSocket(const FName& SocketType, const FString& SocketDescription, ESocketProtocolFamily ProtocolType, bool bForceUDP = false) override;
+	virtual class FSocket* CreateSocket(const FName& SocketType, const FString& SocketDescription, ESocketProtocolFamily ProtocolType) override;
 	virtual TSharedRef<FInternetAddr> GetLocalHostAddr(FOutputDevice& Out, bool& bCanBindAll) override;
 	virtual bool GetLocalAdapterAddresses(TArray<TSharedPtr<FInternetAddr> >& OutAddresses) override;
 };
diff --git a/Engine/Source/Runtime/Sockets/Private/Windows/SocketSubsystemWindows.cpp b/Engine/Source/Runtime/Sockets/Private/Windows/SocketSubsystemWindows.cpp
--- a/Engine/Source/Runtime/Sockets/Private/Windows/SocketSubsystemWindows.cpp
+++ b/Engine/Source/Runtime/Sockets/Private/Windows/SocketSubsystemWindows.cpp
@@ -69,9 +69,9 @@ void FSocketSubsystemWindows::Destroy()
 /* FSocketSubsystemBSD overrides
 *****************************************************************************/
 
-FSocket* FSocketSubsystemWindows::CreateSocket(const FName& SocketType, const FString& SocketDescription, ESocketProtocolFamily ProtocolType, bool bForceUDP)
+FSocket* FSocketSubsystemWindows::CreateSocket(const FName& SocketType, const FString& SocketDescription, ESocketProtocolFamily ProtocolType)
 {
-	FSocketBSD* NewSocket = (FSocketBSD*)FSocketSubsystemBSD::CreateSocket(SocketType, SocketDescription, ProtocolType, bForceUDP);
+	FSocketBSD* NewSocket = (FSocketBSD*)FSocketSubsystemBSD::CreateSocket(SocketType, SocketDescription, ProtocolType);
 
 	if (NewSocket != nullptr)
 	{
diff --git a/Engine/Source/Runtime/Sockets/Private/Windows/SocketSubsystemWindows.h b/Engine/Source/Runtime/Sockets/Private/Windows/SocketSubsystemWindows.h
--- a/Engine/Source/Runtime/Sockets/Private/Windows/SocketSubsystemWindows.h
+++ b/Engine/Source/Runtime/Sockets/Private/Windows/SocketSubsystemWindows.h
@@ -28,7 +28,7 @@ public:
 
 	// FSocketSubsystemBSD overrides
 
-	virtual class FSocket* CreateSocket(const FName& SocketType, const FString& SocketDescription, ESocketProtocolFamily ProtocolType, bool bForceUDP = false) override;
+	virtual class FSocket* CreateSocket(const FName& SocketType, const FString& SocketDescription, ESocketProtocolFamily ProtocolType) override;
 	virtual bool HasNetworkDevice() override;
 	virtual ESocketErrors GetLastErrorCode() override;
 	virtual bool GetLocalAdapterAddresses( TArray<TSharedPtr<FInternetAddr> >& OutAdresses ) override;
diff --git a/Engine/Source/Runtime/Sockets/Public/SocketSubsystem.h b/Engine/Source/Runtime/Sockets/Public/SocketSubsystem.h
--- a/Engine/Source/Runtime/Sockets/Public/SocketSubsystem.h
+++ b/Engine/Source/Runtime/Sockets/Public/SocketSubsystem.h
@@ -84,7 +84,7 @@ public:
 	 */
 	virtual class FSocket* CreateSocket(const FName& SocketType, const FString& SocketDescription, bool bForceUDP = false)
 	{
-		return CreateSocket(SocketType, SocketDescription, ESocketProtocolFamily::None, bForceUDP);
+		return CreateSocket(SocketType, SocketDescription, ESocketProtocolFamily::None);
 	}
 
 	/**
@@ -93,11 +93,10 @@ public:
 	 * @Param SocketType type of socket to create (DGram, Stream, etc)
 	 * @param SocketDescription debug description
 	 * @param ProtocolType the socket protocol to be used
-	 * @param bForceUDP overrides any platform specific protocol with UDP instead
 	 *
 	 * @return the new socket or NULL if failed
 	 */
-	virtual class FSocket* CreateSocket(const FName& SocketType, const FString& SocketDescription, ESocketProtocolFamily ProtocolType, bool bForceUDP = false) = 0;
+	virtual class FSocket* CreateSocket(const FName& SocketType, const FString& SocketDescription, ESocketProtocolFamily ProtocolType) = 0;
 
 	/**
 	 * Creates a resolve info cached struct to hold the resolved address
