Fix the writer's size issue in the auth handler when requesting resends. This fixes packet over/underflows. Also make it so that dropall drops all of auth's packets as well.

diff --git a/Engine/Plugins/Online/OnlineSubsystemSteam/Source/Private/OnlineAuthHandlerSteam.cpp b/Engine/Plugins/Online/OnlineSubsystemSteam/Source/Private/OnlineAuthHandlerSteam.cpp
--- a/Engine/Plugins/Online/OnlineSubsystemSteam/Source/Private/OnlineAuthHandlerSteam.cpp
+++ b/Engine/Plugins/Online/OnlineSubsystemSteam/Source/Private/OnlineAuthHandlerSteam.cpp
@@ -233,6 +233,11 @@ void FSteamAuthHandlerComponent::SendPacket(FBitWriter& OutboundPacket)
 		OutboundPacket.SetError();
 	}
 
+	if (AuthInterface->bDropAll)
+	{
+		return;
+	}
+
 	if (AuthInterface->bRandomDrop && FMath::RandBool() == false)
 	{
 		UE_LOG_ONLINE(Warning, TEXT("AUTH HANDLER: Random packet was dropped!"));
@@ -248,7 +253,7 @@ void FSteamAuthHandlerComponent::SendPacket(FBitWriter& OutboundPacket)
 
 void FSteamAuthHandlerComponent::RequestResend()
 {
-	FBitWriter ResendWriter(sizeof(FSteamAuthInfoData) + 1);
+	FBitWriter ResendWriter(sizeof(FSteamAuthInfoData) * 8 + 1);
 	FSteamAuthInfoData ResendingPacket;
 
 	ResendWriter.WriteBit(1);
