Stop Dedicated Servers from matchmaking as that's not a feature we support. This prevents crashes.

diff --git a/Engine/Plugins/Online/OnlineSubsystemSteam/Source/Private/OnlineSessionInterfaceSteam.cpp b/Engine/Plugins/Online/OnlineSubsystemSteam/Source/Private/OnlineSessionInterfaceSteam.cpp
--- a/Engine/Plugins/Online/OnlineSubsystemSteam/Source/Private/OnlineSessionInterfaceSteam.cpp
+++ b/Engine/Plugins/Online/OnlineSubsystemSteam/Source/Private/OnlineSessionInterfaceSteam.cpp
@@ -685,6 +685,14 @@ bool FOnlineSessionSteam::FindSessions(int32 SearchingPlayerNum, const TSharedRe
 {
 	uint32 Return = E_FAIL;
 
+	// Dedicated servers shouldn't be matchmaking. 
+	if (SteamSubsystem->IsDedicated())
+	{
+		SearchSettings->SearchState = EOnlineAsyncTaskState::Failed;
+		TriggerOnFindSessionsCompleteDelegates(false);
+		return false;
+	}
+
 	// Don't start another search while one is in progress
 	if (!CurrentSessionSearch.IsValid() && SearchSettings->SearchState != EOnlineAsyncTaskState::InProgress)
 	{
