Making it so that the shipping application and dedicated server configurations can be specified in the Target.cs  instead of requiring random edits in the cpps somewhere.

Change requires some new documentation as to how to get started with Steam on UE4.

diff --git a/Engine/Plugins/Online/OnlineSubsystemSteam/Source/Private/OnlineSessionAsyncServerSteam.cpp b/Engine/Plugins/Online/OnlineSubsystemSteam/Source/Private/OnlineSessionAsyncServerSteam.cpp
--- a/Engine/Plugins/Online/OnlineSubsystemSteam/Source/Private/OnlineSessionAsyncServerSteam.cpp
+++ b/Engine/Plugins/Online/OnlineSubsystemSteam/Source/Private/OnlineSessionAsyncServerSteam.cpp
@@ -14,10 +14,23 @@
 /** Turn on Steam filter generation output */
 #define DEBUG_STEAM_FILTERS 1
 
-/** @TODO ONLINE Server values needed to advertise with Steam (NOTE: Steam expects UTF8) */
-#define STEAMPRODUCTNAME "unrealdk"
-#define STEAMGAMEDIR "unrealtest"
-#define STEAMGAMEDESC "Unreal Test!"
+/* Server values needed to advertise with Steam (NOTE: Steam expects UTF8)
+ *
+ * Specify these in your Target.cs files for your project.
+ * If these do not match the same values on the Steam partner backend, 
+ * matchmaking will not work for your dedicated server
+ */
+#ifndef UE4_PROJECT_STEAMPRODUCTNAME
+#define UE4_PROJECT_STEAMPRODUCTNAME "unrealdk"
+#endif
+
+#ifndef UE4_PROJECT_STEAMGAMEDIR
+#define UE4_PROJECT_STEAMGAMEDIR "unrealtest"
+#endif
+
+#ifndef UE4_PROJECT_STEAMGAMEDESC
+#define UE4_PROJECT_STEAMGAMEDESC "Unreal Test!"
+#endif
 
 /**
  * Get the engine unique build id as Steam key
@@ -337,9 +350,9 @@ void FOnlineAsyncTaskSteamCreateServer::Tick()
 			bool bWantsDedicated = Session->SessionSettings.bIsDedicated;
 			UE_LOG_ONLINE(Verbose, TEXT("Initializing Steam game server. Is dedicated? %d"), bWantsDedicated);
 
-			SteamGameServerPtr->SetModDir(STEAMGAMEDIR);
-			SteamGameServerPtr->SetProduct(STEAMPRODUCTNAME);
-			SteamGameServerPtr->SetGameDescription(STEAMGAMEDESC);
+			SteamGameServerPtr->SetModDir(UE4_PROJECT_STEAMGAMEDIR);
+			SteamGameServerPtr->SetProduct(UE4_PROJECT_STEAMPRODUCTNAME);
+			SteamGameServerPtr->SetGameDescription(UE4_PROJECT_STEAMGAMEDESC);
 			SteamGameServerPtr->SetDedicatedServer(bWantsDedicated);
 
 			if (!SteamGameServerPtr->BLoggedOn())
diff --git a/Engine/Plugins/Online/OnlineSubsystemSteam/Source/Private/OnlineSubsystemSteam.cpp b/Engine/Plugins/Online/OnlineSubsystemSteam/Source/Private/OnlineSubsystemSteam.cpp
--- a/Engine/Plugins/Online/OnlineSubsystemSteam/Source/Private/OnlineSubsystemSteam.cpp
+++ b/Engine/Plugins/Online/OnlineSubsystemSteam/Source/Private/OnlineSubsystemSteam.cpp
@@ -31,6 +31,14 @@
 #include "OnlineAuthInterfaceSteam.h"
 #include "VoiceInterfaceSteam.h"
 
+/* Specify this define in your Target.cs for your project
+ *
+ * This helps the SteamAPI find your project on shipping builds
+ * if your game is launched outside of Steam.
+ */
+#ifndef UE4_PROJECT_STEAMSHIPPINGID
+#define UE4_PROJECT_STEAMSHIPPINGID 0
+#endif
 
 #if !UE_BUILD_SHIPPING
 namespace OSSConsoleVariables
@@ -188,8 +196,7 @@ void ConfigureSteamInitDevOptions(bool& RequireRelaunch, int32& RelaunchAppId)
 #else
 	// Always check against the Steam client when shipping
 	RequireRelaunch = true;
-	// Enter shipping app id here
-	RelaunchAppId = 0;
+	RelaunchAppId = UE4_PROJECT_STEAMSHIPPINGID;
 #endif
 }
 
