Quick fix to make the Target.cs strings actually strings. Should have been a part of the last commit but didn't merge the change properly.

diff --git a/Engine/Plugins/Online/OnlineSubsystemSteam/Source/Private/OnlineSessionAsyncServerSteam.cpp b/Engine/Plugins/Online/OnlineSubsystemSteam/Source/Private/OnlineSessionAsyncServerSteam.cpp
--- a/Engine/Plugins/Online/OnlineSubsystemSteam/Source/Private/OnlineSessionAsyncServerSteam.cpp
+++ b/Engine/Plugins/Online/OnlineSubsystemSteam/Source/Private/OnlineSessionAsyncServerSteam.cpp
@@ -21,15 +21,15 @@
  * matchmaking will not work for your dedicated server
  */
 #ifndef UE4_PROJECT_STEAMPRODUCTNAME
-#define UE4_PROJECT_STEAMPRODUCTNAME "unrealdk"
+#define UE4_PROJECT_STEAMPRODUCTNAME unrealdk
 #endif
 
 #ifndef UE4_PROJECT_STEAMGAMEDIR
-#define UE4_PROJECT_STEAMGAMEDIR "unrealtest"
+#define UE4_PROJECT_STEAMGAMEDIR unrealtest
 #endif
 
 #ifndef UE4_PROJECT_STEAMGAMEDESC
-#define UE4_PROJECT_STEAMGAMEDESC "Unreal Test!"
+#define UE4_PROJECT_STEAMGAMEDESC Unreal Test!
 #endif
 
 /**
@@ -350,9 +350,9 @@ void FOnlineAsyncTaskSteamCreateServer::Tick()
 			bool bWantsDedicated = Session->SessionSettings.bIsDedicated;
 			UE_LOG_ONLINE(Verbose, TEXT("Initializing Steam game server. Is dedicated? %d"), bWantsDedicated);
 
-			SteamGameServerPtr->SetModDir(UE4_PROJECT_STEAMGAMEDIR);
-			SteamGameServerPtr->SetProduct(UE4_PROJECT_STEAMPRODUCTNAME);
-			SteamGameServerPtr->SetGameDescription(UE4_PROJECT_STEAMGAMEDESC);
+			SteamGameServerPtr->SetModDir(PREPROCESSOR_TO_STRING(UE4_PROJECT_STEAMGAMEDIR));
+			SteamGameServerPtr->SetProduct(PREPROCESSOR_TO_STRING(UE4_PROJECT_STEAMPRODUCTNAME));
+			SteamGameServerPtr->SetGameDescription(PREPROCESSOR_TO_STRING(UE4_PROJECT_STEAMGAMEDESC));
 			SteamGameServerPtr->SetDedicatedServer(bWantsDedicated);
 
 			if (!SteamGameServerPtr->BLoggedOn())
@@ -786,7 +786,7 @@ void FOnlineAsyncTaskSteamFindServerBase::CreateQuery(MatchMakingKeyValuePair_t*
     NumFilters = 0;
 	// Filter must match at least our game
 	FCStringAnsi::Strncpy(Filters[NumFilters].m_szKey, "gamedir", KeySize);
-	FCStringAnsi::Strncpy(Filters[NumFilters].m_szValue, STEAMGAMEDIR, ValueSize);
+	FCStringAnsi::Strncpy(Filters[NumFilters].m_szValue, PREPROCESSOR_TO_STRING(UE4_PROJECT_STEAMGAMEDIR), ValueSize);
 	NumFilters++;
 
 	FString MapName;
