PR #5207: Only warn about steam being disabled if not in editor mode (Contributed by DSDambuster)

diff --git a/Engine/Plugins/Online/OnlineSubsystemSteam/Source/Private/OnlineSubsystemModuleSteam.cpp b/Engine/Plugins/Online/OnlineSubsystemSteam/Source/Private/OnlineSubsystemModuleSteam.cpp
--- a/Engine/Plugins/Online/OnlineSubsystemSteam/Source/Private/OnlineSubsystemModuleSteam.cpp
+++ b/Engine/Plugins/Online/OnlineSubsystemSteam/Source/Private/OnlineSubsystemModuleSteam.cpp
@@ -60,7 +60,7 @@ public:
 			}
 			else
 			{
-				UE_LOG_ONLINE(Warning, TEXT("Steam API disabled!"));
+				UE_CLOG_ONLINE(IsRunningDedicatedServer() || IsRunningGame(), Warning, TEXT("Steam API disabled!"));
 				DestroySubsystem();
 			}
 
