Standard is to escape these strings already, so change it back to the old method.

#rnx

diff --git a/Engine/Plugins/Online/OnlineSubsystemSteam/Source/Private/OnlineSessionAsyncServerSteam.cpp b/Engine/Plugins/Online/OnlineSubsystemSteam/Source/Private/OnlineSessionAsyncServerSteam.cpp
--- a/Engine/Plugins/Online/OnlineSubsystemSteam/Source/Private/OnlineSessionAsyncServerSteam.cpp
+++ b/Engine/Plugins/Online/OnlineSubsystemSteam/Source/Private/OnlineSessionAsyncServerSteam.cpp
@@ -17,19 +17,20 @@
 /* Server values needed to advertise with Steam (NOTE: Steam expects UTF8)
  *
  * Specify these in your Target.cs files for your project.
+ * Make sure to escape the strings!
  * If these do not match the same values on the Steam partner backend, 
  * matchmaking will not work for your dedicated server
  */
 #ifndef UE4_PROJECT_STEAMPRODUCTNAME
-#define UE4_PROJECT_STEAMPRODUCTNAME unrealdk
+#define UE4_PROJECT_STEAMPRODUCTNAME "unrealdk"
 #endif
 
 #ifndef UE4_PROJECT_STEAMGAMEDIR
-#define UE4_PROJECT_STEAMGAMEDIR unrealtest
+#define UE4_PROJECT_STEAMGAMEDIR "unrealtest"
 #endif
 
 #ifndef UE4_PROJECT_STEAMGAMEDESC
-#define UE4_PROJECT_STEAMGAMEDESC Unreal Test!
+#define UE4_PROJECT_STEAMGAMEDESC "Unreal Test!"
 #endif
 
 /**
@@ -350,9 +351,9 @@ void FOnlineAsyncTaskSteamCreateServer::Tick()
 			bool bWantsDedicated = Session->SessionSettings.bIsDedicated;
 			UE_LOG_ONLINE(Verbose, TEXT("Initializing Steam game server. Is dedicated? %d"), bWantsDedicated);
 
-			SteamGameServerPtr->SetModDir(PREPROCESSOR_TO_STRING(UE4_PROJECT_STEAMGAMEDIR));
-			SteamGameServerPtr->SetProduct(PREPROCESSOR_TO_STRING(UE4_PROJECT_STEAMPRODUCTNAME));
-			SteamGameServerPtr->SetGameDescription(PREPROCESSOR_TO_STRING(UE4_PROJECT_STEAMGAMEDESC));
+			SteamGameServerPtr->SetModDir(UE4_PROJECT_STEAMGAMEDIR);
+			SteamGameServerPtr->SetProduct(UE4_PROJECT_STEAMPRODUCTNAME);
+			SteamGameServerPtr->SetGameDescription(UE4_PROJECT_STEAMGAMEDESC);
 			SteamGameServerPtr->SetDedicatedServer(bWantsDedicated);
 
 			if (!SteamGameServerPtr->BLoggedOn())
@@ -786,7 +787,7 @@ void FOnlineAsyncTaskSteamFindServerBase::CreateQuery(MatchMakingKeyValuePair_t*
     NumFilters = 0;
 	// Filter must match at least our game
 	FCStringAnsi::Strncpy(Filters[NumFilters].m_szKey, "gamedir", KeySize);
-	FCStringAnsi::Strncpy(Filters[NumFilters].m_szValue, PREPROCESSOR_TO_STRING(UE4_PROJECT_STEAMGAMEDIR), ValueSize);
+	FCStringAnsi::Strncpy(Filters[NumFilters].m_szValue, UE4_PROJECT_STEAMGAMEDIR, ValueSize);
 	NumFilters++;
 
 	FString MapName;
