Allowing the Steam SocketSubsystem to be not set as the default, thus overriding the platform socketsubsystem's functionality.

diff --git a/Engine/Plugins/Online/OnlineSubsystemSteam/Source/Private/SocketSubsystemSteam.cpp b/Engine/Plugins/Online/OnlineSubsystemSteam/Source/Private/SocketSubsystemSteam.cpp
--- a/Engine/Plugins/Online/OnlineSubsystemSteam/Source/Private/SocketSubsystemSteam.cpp
+++ b/Engine/Plugins/Online/OnlineSubsystemSteam/Source/Private/SocketSubsystemSteam.cpp
@@ -21,7 +21,7 @@ FName CreateSteamSocketSubsystem()
 	if (SocketSubsystem->Init(Error))
 	{
 		FSocketSubsystemModule& SSS = FModuleManager::LoadModuleChecked<FSocketSubsystemModule>("Sockets");
-		SSS.RegisterSocketSubsystem(STEAM_SUBSYSTEM, SocketSubsystem, true);
+		SSS.RegisterSocketSubsystem(STEAM_SUBSYSTEM, SocketSubsystem, SocketSubsystem->ShouldOverrideDefaultSubsystem());
 		return STEAM_SUBSYSTEM;
 	}
 	else
@@ -489,6 +489,28 @@ void FSocketSubsystemSteam::P2PRemove(FUniqueNetIdSteam& SessionId, int32 Channe
 	AcceptedConnections.Remove(SessionId);
 }
 
+/**
+ * Determines if the SocketSubsystemSteam should override the platform
+ * socket subsystem. This means ISocketSubsystem::Get() will return this subsystem
+ * by default. However, the platform subsystem will still be accessible by
+ * specifying ISocketSubsystem::Get(PLATFORM_SOCKETSUBSYSTEM) as well as via
+ * passthrough operations.
+ *
+ * If the project does not want to use SteamNetworking features, add
+ * bUseSteamNetworking=false to your OnlineSubsystemSteam configuration
+ *
+ * @return if SteamNetworking should be the default socketsubsystem.
+ */
+bool FSocketSubsystemSteam::ShouldOverrideDefaultSubsystem() const
+{
+	bool bOverrideSetting;
+	if (GConfig && GConfig->GetBool(TEXT("OnlineSubsystemSteam"), TEXT("bUseSteamNetworking"), bOverrideSetting, GEngineIni))
+	{
+		return bOverrideSetting;
+	}
+	return true;
+}
+
 /**
  * Chance for the socket subsystem to get some time
  *
diff --git a/Engine/Plugins/Online/OnlineSubsystemSteam/Source/Private/SocketSubsystemSteam.h b/Engine/Plugins/Online/OnlineSubsystemSteam/Source/Private/SocketSubsystemSteam.h
--- a/Engine/Plugins/Online/OnlineSubsystemSteam/Source/Private/SocketSubsystemSteam.h
+++ b/Engine/Plugins/Online/OnlineSubsystemSteam/Source/Private/SocketSubsystemSteam.h
@@ -172,6 +172,20 @@ PACKAGE_SCOPE:
 	 */
 	void P2PRemove(FUniqueNetIdSteam& SessionId, int32 Channel = -1);
 
+	/**
+	 * Determines if the SocketSubsystemSteam should override the platform
+	 * socket subsystem. This means ISocketSubsystem::Get() will return this subsystem
+	 * by default. However, the platform subsystem will still be accessible by
+	 * specifying ISocketSubsystem::Get(PLATFORM_SOCKETSUBSYSTEM) as well as via
+	 * passthrough operations.
+	 *
+	 * If the project does not want to use SteamNetworking features, add
+	 * bUseSteamNetworking=false to your OnlineSubsystemSteam configuration
+	 *
+	 * @return if SteamNetworking should be the default socketsubsystem.
+	 */
+	bool ShouldOverrideDefaultSubsystem() const;
+
 public:
 
 	FSocketSubsystemSteam() :
