Fix Steam DS LAN crashing by sending out the Steam LAN account constant instead of null. This prevents the crash, but additional work is needed to connect in this scenario.

diff --git a/Engine/Plugins/Online/OnlineSubsystemSteam/Source/Private/OnlineSessionInterfaceSteam.cpp b/Engine/Plugins/Online/OnlineSubsystemSteam/Source/Private/OnlineSessionInterfaceSteam.cpp
--- a/Engine/Plugins/Online/OnlineSubsystemSteam/Source/Private/OnlineSessionInterfaceSteam.cpp
+++ b/Engine/Plugins/Online/OnlineSubsystemSteam/Source/Private/OnlineSessionInterfaceSteam.cpp
@@ -228,7 +228,7 @@ bool FOnlineSessionSteam::CreateSession(int32 HostingPlayerNum, FName SessionNam
 		Session->NumOpenPublicConnections = NewSessionSettings.bIsDedicated ? NewSessionSettings.NumPublicConnections : NewSessionSettings.NumPublicConnections - 1;
 
 		Session->HostingPlayerNum = HostingPlayerNum;
-		Session->OwningUserId = SteamUser() ? MakeShareable(new FUniqueNetIdSteam(SteamUser()->GetSteamID())) : NULL;
+		Session->OwningUserId = SteamUser() ? MakeShareable(new FUniqueNetIdSteam(SteamUser()->GetSteamID())) : nullptr;
 		Session->OwningUserName = SteamFriends() ? SteamFriends()->GetPersonaName() : FString(TEXT(""));
 		
 		// Unique identifier of this build for compatibility
@@ -350,6 +350,12 @@ uint32 FOnlineSessionSteam::CreateLANSession(int32 HostingPlayerNum, FNamedOnlin
 	// Setup the host session info
 	FOnlineSessionInfoSteam* NewSessionInfo = new FOnlineSessionInfoSteam(ESteamSession::LANSession);
 	NewSessionInfo->InitLAN();
+
+	if (!Session->OwningUserId.IsValid())
+	{
+		// Use the lan user id, requires us to advertise the host ip and port
+		Session->OwningUserId = MakeShareable(new FUniqueNetIdSteam(k_steamIDLanModeGS));
+	}
 	Session->SessionInfo = MakeShareable(NewSessionInfo);
 
 	// Don't create the LAN beacon if advertising is off
