Backwards compatibility for the dedicated server macros for Steam (so we don't break anyone and alert anyone who has defined these in their Target.cs files already).

diff --git a/Engine/Plugins/Online/OnlineSubsystemSteam/Source/Private/OnlineSessionAsyncServerSteam.cpp b/Engine/Plugins/Online/OnlineSubsystemSteam/Source/Private/OnlineSessionAsyncServerSteam.cpp
--- a/Engine/Plugins/Online/OnlineSubsystemSteam/Source/Private/OnlineSessionAsyncServerSteam.cpp
+++ b/Engine/Plugins/Online/OnlineSubsystemSteam/Source/Private/OnlineSessionAsyncServerSteam.cpp
@@ -22,16 +22,31 @@
  * matchmaking will not work for your dedicated server
  */
 #ifndef UE4_PROJECT_STEAMPRODUCTNAME
+#ifdef STEAMPRODUCTNAME
+UE_DEPRECATED(4.22, "The Steam Product Name Macro has been updated to be configurable from your Target.cs file. Please change STEAMPRODUCTNAME to UE4_PROJECT_STEAMPRODUCTNAME.")
+#define UE4_PROJECT_STEAMPRODUCTNAME STEAMPRODUCTNAME
+#else
 #define UE4_PROJECT_STEAMPRODUCTNAME "unrealdk"
 #endif
+#endif
 
 #ifndef UE4_PROJECT_STEAMGAMEDIR
+#ifdef STEAMGAMEDIR
+UE_DEPRECATED(4.22, "The Steam Game Directory Macro has been updated to be configurable from your Target.cs file. Please change STEAMGAMEDIR to UE4_PROJECT_STEAMGAMEDIR.")
+#define UE4_PROJECT_STEAMGAMEDIR STEAMGAMEDIR
+#else
 #define UE4_PROJECT_STEAMGAMEDIR "unrealtest"
 #endif
+#endif
 
 #ifndef UE4_PROJECT_STEAMGAMEDESC
+#ifdef STEAMGAMEDESC
+UE_DEPRECATED(4.22, "The Steam Game Description Macro has been updated to be configurable from your Target.cs file. Please change STEAMGAMEDESC to UE4_PROJECT_STEAMGAMEDESC.")
+#define UE4_PROJECT_STEAMGAMEDESC STEAMGAMEDESC
+#else
 #define UE4_PROJECT_STEAMGAMEDESC "Unreal Test!"
 #endif
+#endif
 
 /**
  * Get the engine unique build id as Steam key
@@ -349,7 +364,11 @@ void FOnlineAsyncTaskSteamCreateServer::Tick()
 		if (Session != nullptr && SteamGameServerPtr != nullptr)
 		{
 			bool bWantsDedicated = Session->SessionSettings.bIsDedicated;
-			UE_LOG_ONLINE(Verbose, TEXT("Initializing Steam game server. Is dedicated? %d"), bWantsDedicated);
+			UE_LOG_ONLINE(Verbose, TEXT("Starting Steam game server. Dedicated? %d Game Dir is: %s Product Name is: %s\nGame Desc is: %s"),
+				bWantsDedicated,
+				ANSI_TO_TCHAR((UE4_PROJECT_STEAMGAMEDIR)),
+				ANSI_TO_TCHAR((UE4_PROJECT_STEAMPRODUCTNAME)),
+				ANSI_TO_TCHAR((UE4_PROJECT_STEAMGAMEDESC)));
 
 			SteamGameServerPtr->SetModDir(UE4_PROJECT_STEAMGAMEDIR);
 			SteamGameServerPtr->SetProduct(UE4_PROJECT_STEAMPRODUCTNAME);
