Change Steam advertising filters such that we don't overflow.

Here are the new changes:

SetGameTags: ViaOnlineService (just this alone). It contains session flags and the build id as specialized cases.
SetGameData: ViaOnlineService + ViaOnlineServiceAndPing (No change)
SetKeyValue: ViaOnlineService + ViaPing + ViaOnlineServiceAndPing (No change)

diff --git a/Engine/Plugins/Online/OnlineSubsystemSteam/Source/Private/OnlineSessionAsyncServerSteam.cpp b/Engine/Plugins/Online/OnlineSubsystemSteam/Source/Private/OnlineSessionAsyncServerSteam.cpp
--- a/Engine/Plugins/Online/OnlineSubsystemSteam/Source/Private/OnlineSessionAsyncServerSteam.cpp
+++ b/Engine/Plugins/Online/OnlineSubsystemSteam/Source/Private/OnlineSessionAsyncServerSteam.cpp
@@ -229,66 +229,27 @@ void UpdatePublishedSettings(UWorld* World, FNamedOnlineSession* Session)
 	}
 
 	// Get the advertised session settings out as Steam key/value pairs
-	FSteamSessionKeyValuePairs AdvertisedKeyValuePairs;
-	GetServerKeyValuePairsFromSession(Session, AdvertisedKeyValuePairs);
+	FSteamSessionKeyValuePairs ViaOnlineServicePairs;
+	GetServerKeyValuePairsFromSession(Session, ViaOnlineServicePairs);
 
-	if (Session->SessionInfo.IsValid())
-	{
-		FOnlineSessionInfoSteam* SessionInfo = (FOnlineSessionInfoSteam*)(Session->SessionInfo.Get());
-		GetServerKeyValuePairsFromSessionInfo(SessionInfo, AdvertisedKeyValuePairs);
-	}
-
-	FString SessionFlags = GetSessionFlagsAsString(Session->SessionSettings);
-	AdvertisedKeyValuePairs.Add(STEAMKEY_SESSIONFLAGS, SessionFlags);
-
-	FString SessionBuildUniqueId = GetBuildIdAsSteamKey(Session->SessionSettings);
-
-	GetServerKeyValuePairsFromSessionSettings(TempSessionSettings, AdvertisedKeyValuePairs, EOnlineDataAdvertisementType::ViaOnlineService);
-
-	FSteamSessionKeyValuePairs AuxKeyValuePairs;
-	GetServerKeyValuePairsFromSessionSettings(TempSessionSettings, AuxKeyValuePairs, EOnlineDataAdvertisementType::ViaPingOnly);
-	
-	FSteamSessionKeyValuePairs TempKeyValuePairs;
-	GetServerKeyValuePairsFromSessionSettings(TempSessionSettings, TempKeyValuePairs, EOnlineDataAdvertisementType::ViaOnlineServiceAndPing);
-
-	AdvertisedKeyValuePairs.Append(TempKeyValuePairs);
-	AuxKeyValuePairs.Append(TempKeyValuePairs);
-	
-	FString GameTagsString, GameDataString;
+	// Grab the session flags
+	ViaOnlineServicePairs.Add(STEAMKEY_SESSIONFLAGS, GetSessionFlagsAsString(Session->SessionSettings));
 
 	// Start the game tags with the build id so search results can early out
-	GameTagsString = SessionBuildUniqueId;
+	FString GameTagsString = GetBuildIdAsSteamKey(Session->SessionSettings);
 
-	// Create the properly formatted Steam string (ie key:value,key:value,key) for GameTags/GameData
-	FSteamSessionKeyValuePairs::TConstIterator It(AdvertisedKeyValuePairs);
-	if (It)
-	{
-		UE_LOG_ONLINE_SESSION(Verbose, TEXT("Master Server Data (%s, %s)"), *It.Key(), *It.Value());
-		FString NewKey = FString::Printf(TEXT("%s:%s"), *It.Key(), *It.Value());
+	GetServerKeyValuePairsFromSessionSettings(TempSessionSettings, ViaOnlineServicePairs, EOnlineDataAdvertisementType::ViaOnlineService);
 
-		if (GameTagsString.Len() + NewKey.Len() < k_cbMaxGameServerTags)
-		{
-			GameTagsString = GameTagsString + "," + NewKey;
-		}
-		else
-		{
-			UE_LOG_ONLINE_SESSION(Warning, TEXT("Server setting %s overflows Steam SetGameTags call"), *NewKey);
-		}
+	FSteamSessionKeyValuePairs ViaPingPairs;
+	GetServerKeyValuePairsFromSessionSettings(TempSessionSettings, ViaPingPairs, EOnlineDataAdvertisementType::ViaPingOnly);
 
-		if (NewKey.Len() < k_cbMaxGameServerGameData)
-		{
-			GameDataString = NewKey;
-		}
-		else
-		{
-			UE_LOG_ONLINE_SESSION(Warning, TEXT("Server setting %s overflows Steam SetGameData call"), *NewKey);
-		}
+	FSteamSessionKeyValuePairs ViaOnlineServicePingPairs;
+	GetServerKeyValuePairsFromSessionSettings(TempSessionSettings, ViaOnlineServicePingPairs, EOnlineDataAdvertisementType::ViaOnlineServiceAndPing);
 
-		++It;
-	}
-	for (; It; ++It)
+	// Generate the tag data first, due to its small size
+	for (FSteamSessionKeyValuePairs::TConstIterator It(ViaOnlineServicePairs); It; ++It)
 	{
-		UE_LOG_ONLINE_SESSION(Verbose, TEXT("Master Server Data (%s, %s)"), *It.Key(), *It.Value());
+		UE_LOG_ONLINE_SESSION(Verbose, TEXT("Master Server Game Tags (%s, %s)"), *It.Key(), *It.Value());
 		FString NewKey = FString::Printf(TEXT(",%s:%s"), *It.Key(), *It.Value());
 		if (GameTagsString.Len() + NewKey.Len() < k_cbMaxGameServerTags)
 		{
@@ -298,10 +259,28 @@ void UpdatePublishedSettings(UWorld* World, FNamedOnlineSession* Session)
 		{
 			UE_LOG_ONLINE_SESSION(Warning, TEXT("Server setting %s overflows Steam SetGameTags call"), *NewKey);
 		}
+	}
+
+	// Push in the other tags from session info (these are dropped from the tag queries for space reasons)
+	if (Session->SessionInfo.IsValid())
+	{
+		FOnlineSessionInfoSteam* SessionInfo = (FOnlineSessionInfoSteam*)(Session->SessionInfo.Get());
+		GetServerKeyValuePairsFromSessionInfo(SessionInfo, ViaOnlineServicePairs);
+	}
+
+	// Append both advertising sets with the dual advertisement filters
+	ViaOnlineServicePairs.Append(ViaOnlineServicePingPairs);
+
+	// Create the properly formatted Steam string (ie key:value,key:value,key) for GameData
+	FString GameDataString;
+	for (FSteamSessionKeyValuePairs::TConstIterator It(ViaOnlineServicePairs); It; ++It)
+	{
+		UE_LOG_ONLINE_SESSION(Verbose, TEXT("Master Server Game Data (%s, %s)"), *It.Key(), *It.Value());
+		FString NewKey = FString::Printf(TEXT("%s:%s"), *It.Key(), *It.Value());
 
 		if (GameDataString.Len() + NewKey.Len() < k_cbMaxGameServerGameData)
 		{
-			GameDataString += NewKey;
+			GameDataString += ((GameDataString.Len() > 0) ? "," : "") + NewKey;
 		}
 		else
 		{
@@ -316,31 +295,26 @@ void UpdatePublishedSettings(UWorld* World, FNamedOnlineSession* Session)
 		SteamGameServerPtr->SetGameTags(TCHAR_TO_UTF8(*GameTagsString));
 	}
 
-	// Large and searchable game data (never returned)
+	// Large and searchable game data (never returned, used to filter things on the master server)
 	if (GameDataString.Len() > 0 && GameDataString.Len() < k_cbMaxGameServerGameData)
 	{
 		UE_LOG_ONLINE_SESSION(Verbose, TEXT("SetGameData(%s)"), *GameDataString);
 		SteamGameServerPtr->SetGameData(TCHAR_TO_UTF8(*GameDataString));
 	}
 
-	// @TODO ONLINE - distinguish between server side keys (SetGameData()) and client side keys (SetKeyValue())
-	// Set the advertised filter keys (these can not be filtered at master-server level, only client side)
+	// Set the advertised filter keys (these can not be filtered at master-server level, only client side via talking to the server)
 	SteamGameServerPtr->ClearAllKeyValues();
 
-	// Key value pairs sent as rules (requires secondary RulesRequest call)
-	for (FSteamSessionKeyValuePairs::TConstIterator AdvKeyIt(AdvertisedKeyValuePairs); AdvKeyIt; ++AdvKeyIt)
-	{
-		UE_LOG_ONLINE_SESSION(Verbose, TEXT("Aux Server Data (%s, %s)"), *AdvKeyIt.Key(), *AdvKeyIt.Value());
-		SteamGameServerPtr->SetKeyValue(TCHAR_TO_UTF8(*AdvKeyIt.Key()), TCHAR_TO_UTF8(*AdvKeyIt.Value()));
-	}
+	// Push all the filters in for the KeyValue data.
+	ViaOnlineServicePairs.Append(ViaPingPairs);
 
 	// Key value pairs sent as rules (requires secondary RulesRequest call)
-	for (FSteamSessionKeyValuePairs::TConstIterator AuxKeyIt(AuxKeyValuePairs); AuxKeyIt; ++AuxKeyIt)
+	for (FSteamSessionKeyValuePairs::TConstIterator AuxKeyIt(ViaOnlineServicePairs); AuxKeyIt; ++AuxKeyIt)
 	{
-		UE_LOG_ONLINE_SESSION(Verbose, TEXT("Aux Server Data (%s, %s)"), *AuxKeyIt.Key(), *AuxKeyIt.Value());
+		UE_LOG_ONLINE_SESSION(Verbose, TEXT("Pushing Server KVData (%s, %s)"), *AuxKeyIt.Key(), *AuxKeyIt.Value());
 		SteamGameServerPtr->SetKeyValue(TCHAR_TO_UTF8(*AuxKeyIt.Key()), TCHAR_TO_UTF8(*AuxKeyIt.Value()));
 	}
-}	
+}
 
 /**
  *	Get a human readable description of task
