Adding functionality to allow developers/server hosting providers to set Steam dedicated server names/titles at application launch. This allows for custom server names to appear on the external server browser UI (and in game).

This fixes the previous issue where servers running on a headless Steam instance (not logged into an account) would show up in server browsers as just a string of numbers.

Example usage: -SteamServerName="DS Test!"

diff --git a/Engine/Plugins/Online/OnlineSubsystemSteam/Source/Private/OnlineSessionAsyncServerSteam.cpp b/Engine/Plugins/Online/OnlineSubsystemSteam/Source/Private/OnlineSessionAsyncServerSteam.cpp
--- a/Engine/Plugins/Online/OnlineSubsystemSteam/Source/Private/OnlineSessionAsyncServerSteam.cpp
+++ b/Engine/Plugins/Online/OnlineSubsystemSteam/Source/Private/OnlineSessionAsyncServerSteam.cpp
@@ -414,7 +414,10 @@ void FOnlineAsyncTaskSteamCreateServer::Finalize()
 				check(Session->SessionSettings.bIsDedicated);
 				// Associate the dedicated server anonymous login as the owning user
 				Session->OwningUserId = SessionInt->GameServerSteamId;
-				Session->OwningUserName = Session->OwningUserId->ToString();
+
+				// Figure if we need to override the server name here
+				FString CustomDedicatedServerName = SessionInt->GetCustomDedicatedServerName();
+				Session->OwningUserName = (!CustomDedicatedServerName.IsEmpty()) ? CustomDedicatedServerName : Session->OwningUserId->ToString();
 			}
 
 			bool bShouldUseAdvertise = true;
diff --git a/Engine/Plugins/Online/OnlineSubsystemSteam/Source/Private/OnlineSessionInterfaceSteam.cpp b/Engine/Plugins/Online/OnlineSubsystemSteam/Source/Private/OnlineSessionInterfaceSteam.cpp
--- a/Engine/Plugins/Online/OnlineSubsystemSteam/Source/Private/OnlineSessionInterfaceSteam.cpp
+++ b/Engine/Plugins/Online/OnlineSubsystemSteam/Source/Private/OnlineSessionInterfaceSteam.cpp
@@ -229,7 +229,7 @@ bool FOnlineSessionSteam::CreateSession(int32 HostingPlayerNum, FName SessionNam
 
 		Session->HostingPlayerNum = HostingPlayerNum;
 		Session->OwningUserId = SteamUser() ? MakeShareable(new FUniqueNetIdSteam(SteamUser()->GetSteamID())) : nullptr;
-		Session->OwningUserName = SteamFriends() ? SteamFriends()->GetPersonaName() : FString(TEXT(""));
+		Session->OwningUserName = SteamFriends() ? SteamFriends()->GetPersonaName() : GetCustomDedicatedServerName();
 		
 		// Unique identifier of this build for compatibility
 		Session->SessionSettings.BuildUniqueId = GetBuildUniqueId();
@@ -1316,6 +1316,25 @@ bool FOnlineSessionSteam::GetResolvedConnectString(const FOnlineSessionSearchRes
 	return bSuccess;
 }
 
+FString FOnlineSessionSteam::GetCustomDedicatedServerName() const
+{
+	FString ServerName;
+
+	if (FParse::Value(FCommandLine::Get(), TEXT("-SteamServerName="), ServerName))
+	{
+		if (ServerName.Len() >= k_cbMaxGameServerName)
+		{
+			UE_LOG_ONLINE_SESSION(Warning, TEXT("SteamServerName overflows the maximum amount of characters %d allowed, truncating."), k_cbMaxGameServerName);
+			// Must have space for the null terminator
+			ServerName = ServerName.Left(k_cbMaxGameServerName - 1);
+		}
+
+		return ServerName;
+	}
+	
+	return TEXT("");
+}
+
 FOnlineSessionSettings* FOnlineSessionSteam::GetSessionSettings(FName SessionName) 
 {
 	FNamedOnlineSession* Session = GetNamedSession(SessionName);
diff --git a/Engine/Plugins/Online/OnlineSubsystemSteam/Source/Private/OnlineSessionInterfaceSteam.h b/Engine/Plugins/Online/OnlineSubsystemSteam/Source/Private/OnlineSessionInterfaceSteam.h
--- a/Engine/Plugins/Online/OnlineSubsystemSteam/Source/Private/OnlineSessionInterfaceSteam.h
+++ b/Engine/Plugins/Online/OnlineSubsystemSteam/Source/Private/OnlineSessionInterfaceSteam.h
@@ -434,6 +434,15 @@ PACKAGE_SCOPE:
 	 */
 	void RegisterLocalPlayers(class FNamedOnlineSession* Session);
 
+	/**
+	 * Parses the dedicated server custom name launch argument as specified by the
+	 * -SteamServerName= flag.
+	 *
+	 * @return the custom name if the value of the launch argument is specified and less than 64 characters
+	 *			otherwise returns empty string
+	 */
+	FString GetCustomDedicatedServerName() const;
+
 public:
 
 	virtual ~FOnlineSessionSteam() {}
