Fix problems with Steam NetConnections and P2P communications incorrectly storing connection info causing problems when multiple Steam connections to a single user are active.

diff --git a/Engine/Plugins/Online/OnlineSubsystemSteam/Source/Private/IPAddressSteam.h b/Engine/Plugins/Online/OnlineSubsystemSteam/Source/Private/IPAddressSteam.h
--- a/Engine/Plugins/Online/OnlineSubsystemSteam/Source/Private/IPAddressSteam.h
+++ b/Engine/Plugins/Online/OnlineSubsystemSteam/Source/Private/IPAddressSteam.h
@@ -165,7 +165,17 @@ public:
 
 	virtual uint32 GetTypeHash() override
 	{
-		return ::GetTypeHash(*(uint64*)SteamId.GetBytes());
+		return GetConstTypeHash();
+	}
+
+	uint32 GetConstTypeHash() const
+	{
+		return ::GetTypeHash(ToString(true));
+	}
+
+	friend uint32 GetTypeHash(const FInternetAddrSteam& A)
+	{
+		return A.GetConstTypeHash();
 	}
 
 	/**
diff --git a/Engine/Plugins/Online/OnlineSubsystemSteam/Source/Private/SocketSubsystemSteam.cpp b/Engine/Plugins/Online/OnlineSubsystemSteam/Source/Private/SocketSubsystemSteam.cpp
--- a/Engine/Plugins/Online/OnlineSubsystemSteam/Source/Private/SocketSubsystemSteam.cpp
+++ b/Engine/Plugins/Online/OnlineSubsystemSteam/Source/Private/SocketSubsystemSteam.cpp
@@ -4,11 +4,12 @@
 #include "Misc/ConfigCacheIni.h"
 #include "SocketsSteam.h"
 #include "IPAddressSteam.h"
+#include "OnlineSubsystemSteam.h"
 #include "OnlineSessionInterfaceSteam.h"
 #include "SocketSubsystemModule.h"
 #include "SteamNetConnection.h"
 
-FSocketSubsystemSteam* FSocketSubsystemSteam::SocketSingleton = NULL;
+FSocketSubsystemSteam* FSocketSubsystemSteam::SocketSingleton = nullptr;
 
 /**
  * Create the socket subsystem for the given platform service
@@ -52,7 +53,7 @@ void DestroySteamSocketSubsystem()
  */
 FSocketSubsystemSteam* FSocketSubsystemSteam::Create()
 {
-	if (SocketSingleton == NULL)
+	if (SocketSingleton == nullptr)
 	{
 		SocketSingleton = new FSocketSubsystemSteam();
 	}
@@ -65,11 +66,11 @@ FSocketSubsystemSteam* FSocketSubsystemSteam::Create()
  */
 void FSocketSubsystemSteam::Destroy()
 {
-	if (SocketSingleton != NULL)
+	if (SocketSingleton != nullptr)
 	{
 		SocketSingleton->Shutdown();
 		delete SocketSingleton;
-		SocketSingleton = NULL;
+		SocketSingleton = nullptr;
 	}
 }
 
@@ -93,6 +94,11 @@ bool FSocketSubsystemSteam::Init(FString& Error)
 		{
 			UE_LOG_ONLINE(Warning, TEXT("Missing P2PConnectionTimeout key in OnlineSubsystemSteam of DefaultEngine.ini"));
 		}
+
+		if (!GConfig->GetDouble(TEXT("OnlineSubsystemSteam"), TEXT("P2PCleanupTimeout"), P2PCleanupTimeout, GEngineIni))
+		{
+			UE_LOG_ONLINE(Warning, TEXT("Missing P2PCleanupTimeout key in OnlineSubsystemSteam of DefaultEngine.ini"));
+		}
 	}
 
 	if (SteamNetworking())
@@ -122,15 +128,20 @@ void FSocketSubsystemSteam::Shutdown()
 		}
 	}
 
+	UE_LOG_ONLINE(Verbose, TEXT("Shutting down SteamNet connections"));
+	
+	// Empty the DeadConnections list as we're shutting down anyways
+	// This is so we don't spend time checking the DeadConnections
+	// for duplicate pending closures
+	DeadConnections.Empty();
+
 	// Cleanup any remaining sessions
-	TArray<FUniqueNetIdSteam> SessionIds;
-	AcceptedConnections.GenerateKeyArray(SessionIds);
-	for (int SessionIdx=0; SessionIdx < SessionIds.Num(); SessionIdx++)
+	for (auto SessionIds : AcceptedConnections)
 	{
-		P2PRemove(SessionIds[SessionIdx]);
+		P2PRemove(SessionIds.Key, -1);
 	}
-	
-	CleanupDeadConnections();
+
+	CleanupDeadConnections(true);
 
 	// Cleanup sockets
 	TArray<FSocketSteam*> TempArray = SteamSockets;
@@ -156,11 +167,11 @@ void FSocketSubsystemSteam::Shutdown()
  */
 FSocket* FSocketSubsystemSteam::CreateSocket(const FName& SocketType, const FString& SocketDescription, ESocketProtocolFamily ProtocolType)
 {
-	FSocket* NewSocket = NULL;
+	FSocket* NewSocket = nullptr;
 	if (SocketType == FName("SteamClientSocket"))
 	{
 		ISteamUser* SteamUserPtr = SteamUser();
-		if (SteamUserPtr != NULL)
+		if (SteamUserPtr != nullptr)
 		{
 			FUniqueNetIdSteam ClientId(SteamUserPtr->GetSteamID());
 			NewSocket = new FSocketSteam(SteamNetworking(), ClientId, SocketDescription, ProtocolType);
@@ -252,11 +263,13 @@ void FSocketSubsystemSteam::RegisterConnection(USteamNetConnection* Connection)
 	FWeakObjectPtr ObjectPtr = Connection;
 	SteamConnections.Add(ObjectPtr);
 
-	if (Connection->RemoteAddr.IsValid() && Connection->Socket)
+	if (Connection->GetInternetAddr().IsValid() && Connection->Socket)
 	{
 		FSocketSteam* SteamSocket = (FSocketSteam*)Connection->Socket;
-		FInternetAddrSteam& SteamAddr = (FInternetAddrSteam&)(*Connection->RemoteAddr);
-		P2PTouch(SteamSocket->SteamNetworkingPtr, SteamAddr.SteamId);
+		TSharedPtr<const FInternetAddrSteam> SteamAddr = StaticCastSharedPtr<const FInternetAddrSteam>(Connection->GetInternetAddr());
+
+		UE_LOG_ONLINE(Log, TEXT("Adding user %s from RegisterConnection"), *SteamAddr->ToString(true));
+		P2PTouch(SteamSocket->SteamNetworkingPtr, SteamAddr->SteamId, SteamAddr->SteamChannel);
 	}
 }
 
@@ -270,21 +283,20 @@ void FSocketSubsystemSteam::UnregisterConnection(USteamNetConnection* Connection
 	check(!Connection->bIsPassthrough);
 
 	FWeakObjectPtr ObjectPtr = Connection;
-	int32 NumRemoved = SteamConnections.RemoveSingleSwap(ObjectPtr);
 
 	// Don't call P2PRemove again if we didn't actually remove a connection. This 
 	// will get called twice - once the connection is closed and when the connection
 	// is garbage collected. It's possible that the player who left rejoined before garbage
 	// collection runs (their connection object will be different), so P2PRemove would kick
 	// them from the session when it shouldn't.
-	if (NumRemoved > 0 && Connection->RemoteAddr.IsValid())
+	if (SteamConnections.RemoveSingleSwap(ObjectPtr) == 1 && Connection->GetInternetAddr().IsValid())
 	{
-		FInternetAddrSteam& SteamAddr = (FInternetAddrSteam&)(*Connection->RemoteAddr);
-		P2PRemove(SteamAddr.SteamId, SteamAddr.SteamChannel);
+		TSharedPtr<const FInternetAddrSteam> SteamAddr = StaticCastSharedPtr<const FInternetAddrSteam>(Connection->GetInternetAddr());
+		P2PRemove(SteamAddr->SteamId, SteamAddr->SteamChannel);
 	}
 }
 
-void FSocketSubsystemSteam::ConnectFailure(FUniqueNetIdSteam& RemoteId)
+void FSocketSubsystemSteam::ConnectFailure(const FUniqueNetIdSteam& RemoteId)
 {
 	// Remove any GC'd references
 	for (int32 ConnIdx=SteamConnections.Num()-1; ConnIdx>=0; ConnIdx--)
@@ -299,16 +311,16 @@ void FSocketSubsystemSteam::ConnectFailure(FUniqueNetIdSteam& RemoteId)
 	for (int32 ConnIdx=0; ConnIdx<SteamConnections.Num(); ConnIdx++)
 	{
 		USteamNetConnection* SteamConn = CastChecked<USteamNetConnection>(SteamConnections[ConnIdx].Get());
-		FInternetAddrSteam& RemoteAddrSteam = (FInternetAddrSteam&)*(SteamConn->RemoteAddr);
+		TSharedPtr<const FInternetAddrSteam> RemoteAddrSteam = StaticCastSharedPtr<const FInternetAddrSteam>(SteamConn->GetInternetAddr());
 
 		// Only checking Id here because its a complete failure (channel doesn't matter)
-		if (RemoteAddrSteam.SteamId == RemoteId)
+		if (RemoteAddrSteam->SteamId == RemoteId)
 		{
 			SteamConn->Close();
 		}
 	}
 
-	P2PRemove(RemoteId);
+	P2PRemove(RemoteId, -1);
 }
 
 /**
@@ -406,7 +418,7 @@ ESocketErrors FSocketSubsystemSteam::TranslateErrorCode(int32 Code)
  */
 TSharedRef<FInternetAddr> FSocketSubsystemSteam::GetLocalBindAddr(FOutputDevice& Out)
 {
-	FInternetAddrSteam* SteamAddr = NULL;
+	FInternetAddrSteam* SteamAddr = nullptr;
 	CSteamID SteamId;
 	if (SteamUser())
 	{
@@ -437,13 +449,15 @@ TSharedRef<FInternetAddr> FSocketSubsystemSteam::GetLocalBindAddr(FOutputDevice&
  * 
  * @return true if accepted, false otherwise
  */
-bool FSocketSubsystemSteam::AcceptP2PConnection(ISteamNetworking* SteamNetworkingPtr, FUniqueNetIdSteam& RemoteId)
+bool FSocketSubsystemSteam::AcceptP2PConnection(ISteamNetworking* SteamNetworkingPtr, const FUniqueNetIdSteam& RemoteId)
 {
-	if (SteamNetworkingPtr && RemoteId.IsValid() && DeadConnections.Find(RemoteId) == NULL)
+	if (SteamNetworkingPtr && RemoteId.IsValid() && !IsConnectionPendingRemoval(RemoteId, -1))
 	{
+		UE_LOG_ONLINE(Log, TEXT("Adding P2P connection information with user %s (Name: %s)"), *RemoteId.ToString(), *RemoteId.ToDebugString());
 		// Blindly accept connections (but only if P2P enabled)
 		SteamNetworkingPtr->AcceptP2PSessionWithUser(RemoteId);
-		AcceptedConnections.Add(RemoteId, FSteamP2PConnectionInfo(SteamNetworkingPtr, FPlatformTime::Seconds()));
+		UE_CLOG_ONLINE(AcceptedConnections.Contains(RemoteId), Warning, TEXT("User %s already exists in the connections list!!"), *RemoteId.ToString());
+		AcceptedConnections.Add(RemoteId, FSteamP2PConnectionInfo(SteamNetworkingPtr));
 		return true;
 	}
 
@@ -455,15 +469,22 @@ bool FSocketSubsystemSteam::AcceptP2PConnection(ISteamNetworking* SteamNetworkin
  *
  * @param SteamNetworkingPtr proper networking interface that this session is communicating on
  * @param SessionId P2P session recently heard from
+ * @param ChannelId the channel id that the update happened on
  *
  * @return true if the connection is active, false if this is in the dead connections list
  */
-bool FSocketSubsystemSteam::P2PTouch(ISteamNetworking* SteamNetworkingPtr, FUniqueNetIdSteam& SessionId)
+bool FSocketSubsystemSteam::P2PTouch(ISteamNetworking* SteamNetworkingPtr, const FUniqueNetIdSteam& SessionId, int32 ChannelId)
 {
     // Don't update any sessions coming from pending disconnects
-	if (DeadConnections.Find(SessionId) == NULL)
+	if (!IsConnectionPendingRemoval(SessionId, ChannelId))
 	{
-		AcceptedConnections.Add(SessionId, FSteamP2PConnectionInfo(SteamNetworkingPtr, FPlatformTime::Seconds()));
+		FSteamP2PConnectionInfo& ChannelUpdate = AcceptedConnections.FindOrAdd(SessionId);
+		ChannelUpdate.SteamNetworkingPtr = SteamNetworkingPtr;
+
+		if (ChannelId != -1)
+		{
+			ChannelUpdate.AddOrUpdateChannel(ChannelId, FPlatformTime::Seconds());
+		}
 		return true;
 	}
 
@@ -476,17 +497,81 @@ bool FSocketSubsystemSteam::P2PTouch(ISteamNetworking* SteamNetworkingPtr, FUniq
  * @param SessionId P2P session to remove
  * @param Channel channel to close, -1 to close all communication
  */
-void FSocketSubsystemSteam::P2PRemove(FUniqueNetIdSteam& SessionId, int32 Channel)
+void FSocketSubsystemSteam::P2PRemove(const FUniqueNetIdSteam& SessionId, int32 Channel)
 {
 	FSteamP2PConnectionInfo* ConnectionInfo = AcceptedConnections.Find(SessionId);
 	if (ConnectionInfo)
 	{
-        // Move active connections to the dead list so they can be removed (giving Steam a chance to flush connection)
-		DeadConnections.Add(SessionId, FSteamP2PConnectionInfo(ConnectionInfo->SteamNetworkingPtr, FPlatformTime::Seconds(), Channel));
+		const bool bRemoveAllConnections = (Channel == -1);
+		
+		// Only modify the DeadConnections list if we're actively going to change it
+		if (!IsConnectionPendingRemoval(SessionId, Channel))
+		{
+			if (bRemoveAllConnections)
+			{
+				UE_LOG_ONLINE(Verbose, TEXT("Replacing all existing removals with global removal for %s"), *SessionId.ToString());
+				// Go through and remove all the connections for this user
+				for (TMap<FInternetAddrSteam, double>::TIterator It(DeadConnections); It; ++It)
+				{
+					if (It->Key.SteamId == SessionId)
+					{
+						It.RemoveCurrent();
+					}
+				}
+			}
+
+			// Move active connections to the dead list so they can be removed (giving Steam a chance to flush connection)
+			FInternetAddrSteam RemoveConnection(SessionId);
+			RemoveConnection.SetPort(Channel);
+			DeadConnections.Add(RemoveConnection, FPlatformTime::Seconds());
+
+			UE_LOG_ONLINE(Log, TEXT("Removing P2P Session Id: %s, Channel: %d, IdleTime: %0.3f"), *SessionId.ToDebugString(), Channel, 
+				ConnectionInfo ? (FPlatformTime::Seconds() - ConnectionInfo->LastReceivedTime) : 9999.f);
+		}
+
+		if (bRemoveAllConnections)
+		{
+			// Clean up dead connections will remove the user from the map for us
+			UE_CLOG_ONLINE((ConnectionInfo->ConnectedChannels.Num() > 0), Verbose, TEXT("Removing all channel connections for %s"), *SessionId.ToString());
+			ConnectionInfo->ConnectedChannels.Empty();
+		}
+		else
+		{
+			bool bWasRemoved = ConnectionInfo->ConnectedChannels.Remove(Channel) > 0;
+			UE_CLOG_ONLINE(bWasRemoved, Verbose, TEXT("Removing channel %d from user %s"), Channel, *SessionId.ToString());
+		}
+	}
+}
+
+/**
+ * Checks to see if a Steam P2P Connection is pending close on the given channel.
+ *
+ * Before checking the given channel, this function checks if the session is marked for
+ * global disconnection.
+ *
+ * @param SessionId the user id tied to the session disconnection
+ * @param Channel the communications channel id for the user if it exists
+ */
+bool FSocketSubsystemSteam::IsConnectionPendingRemoval(const FUniqueNetIdSteam& SteamId, int32 Channel)
+{
+	FInternetAddrSteam RemovalToFind(SteamId);
+	RemovalToFind.SetPort(-1);
+
+	// Check with -1 first as that ends all communications with another user
+	if (!DeadConnections.Contains(RemovalToFind))
+	{
+		// If we were asked to check for -1, then early out as we've already checked the entry
+		if (Channel == -1)
+		{
+			return false;
+		}
+
+		// Then look for the specific channel instance.
+		RemovalToFind.SetPort(Channel);
+		return DeadConnections.Contains(RemovalToFind);
 	}
 
-	UE_LOG_ONLINE(Log, TEXT("Removing P2P Session Id: %s, IdleTime: %0.3f"), *SessionId.ToDebugString(), ConnectionInfo ? (FPlatformTime::Seconds() - ConnectionInfo->LastReceivedTime) : 9999.f);
-	AcceptedConnections.Remove(SessionId);
+	return true;
 }
 
 /**
@@ -530,7 +615,6 @@ bool FSocketSubsystemSteam::Tick(float DeltaTime)
 		bDumpSessionInfo = true;
 	}
 
-	TArray<FUniqueNetIdSteam> ExpiredSessions;
 	for (TMap<FUniqueNetIdSteam, FSteamP2PConnectionInfo>::TConstIterator It(AcceptedConnections); It; ++It)
 	{
 		const FUniqueNetIdSteam& SessionId = It.Key();
@@ -540,19 +624,19 @@ bool FSocketSubsystemSteam::Tick(float DeltaTime)
 		if (CurSeconds - ConnectionInfo.LastReceivedTime < P2PConnectionTimeout)
 		{
 			P2PSessionState_t SessionInfo;
-			if (ConnectionInfo.SteamNetworkingPtr != NULL && ConnectionInfo.SteamNetworkingPtr->GetP2PSessionState(SessionId, &SessionInfo))
+			if (ConnectionInfo.SteamNetworkingPtr != nullptr && ConnectionInfo.SteamNetworkingPtr->GetP2PSessionState(SessionId, &SessionInfo))
 			{
 				bExpiredSession = false;
 
 				if (bDumpSessionInfo)
 				{
 					UE_LOG_ONLINE(Verbose, TEXT("Dumping Steam P2P socket details:"));
-					UE_LOG_ONLINE(Verbose, TEXT("- Id: %s, IdleTime: %0.3f"), *SessionId.ToDebugString(), (CurSeconds - ConnectionInfo.LastReceivedTime));
+					UE_LOG_ONLINE(Verbose, TEXT("- Id: %s, Number of Channels: %d, IdleTime: %0.3f"), *SessionId.ToDebugString(), ConnectionInfo.ConnectedChannels.Num(), (CurSeconds - ConnectionInfo.LastReceivedTime));
 
 					DumpSteamP2PSessionInfo(SessionInfo);
 				}
 			}
-			else
+			else if(ConnectionInfo.ConnectedChannels.Num() > 0) // Suppress this print so that it only prints if we expected to have a connection.
 			{
 				UE_LOG_ONLINE(Verbose, TEXT("Failed to get Steam P2P session state for Id: %s, IdleTime: %0.3f"), *SessionId.ToDebugString(), (CurSeconds - ConnectionInfo.LastReceivedTime));
 			}
@@ -560,54 +644,80 @@ bool FSocketSubsystemSteam::Tick(float DeltaTime)
 
 		if (bExpiredSession)
 		{
-			ExpiredSessions.Add(SessionId);
+			P2PRemove(SessionId, -1);
 		}
 	}
 
-	// Cleanup any closed/invalid sessions
-	for (int32 SessionIdx=0; SessionIdx < ExpiredSessions.Num(); SessionIdx++)
+	CleanupDeadConnections(false);
+
+	return true;
+}
+
+bool FSocketSubsystemSteam::Exec(UWorld* InWorld, const TCHAR* Cmd, FOutputDevice& Ar)
+{
+#if !UE_BUILD_SHIPPING
+	if (FParse::Command(&Cmd, TEXT("dumpsteamsessions")))
 	{
-		FUniqueNetIdSteam& SessionId = ExpiredSessions[SessionIdx];
-		P2PRemove(SessionId);
+		DumpAllOpenSteamSessions();
+		return true;
 	}
+#endif
 
-	CleanupDeadConnections();
-
-	return true;
+	return false;
 }
 
 /**
  * Iterate through the pending dead connections and permanently remove any that have been around
  * long enough to flush their contents
+ * 
+ * @param bSkipLinger skips the timeout reserved for lingering connection data
  */
-void FSocketSubsystemSteam::CleanupDeadConnections()
+void FSocketSubsystemSteam::CleanupDeadConnections(bool bSkipLinger)
 {
 	double CurSeconds = FPlatformTime::Seconds();
-	TArray<FUniqueNetIdSteam> ExpiredSessions;
-	for (TMap<FUniqueNetIdSteam, FSteamP2PConnectionInfo>::TConstIterator It(DeadConnections); It; ++It)
+	for (TMap<FInternetAddrSteam, double>::TIterator It(DeadConnections); It; ++It)
 	{
-		const FUniqueNetIdSteam& SessionId = It.Key();
-		const FSteamP2PConnectionInfo& ConnectionInfo = It.Value();
-		if (CurSeconds - ConnectionInfo.LastReceivedTime >= 3.0f)
+		const FInternetAddrSteam& SteamConnection = It.Key();
+		if (P2PCleanupTimeout == 0.0 || CurSeconds - It.Value() >= P2PCleanupTimeout || bSkipLinger)
 		{
-			if (ConnectionInfo.Channel == -1)
+			// Only modify connections if the user exists. This check is only done for safety
+			if (AcceptedConnections.Contains(SteamConnection.SteamId))
 			{
-				ConnectionInfo.SteamNetworkingPtr->CloseP2PSessionWithUser(SessionId);
-			}
-			else
-			{
-				ConnectionInfo.SteamNetworkingPtr->CloseP2PChannelWithUser(SessionId, ConnectionInfo.Channel);
+				const FSteamP2PConnectionInfo& ConnectionInfo = AcceptedConnections[SteamConnection.SteamId];
+				bool bShouldRemoveUser = true;
+				// All communications are to be removed
+				if (SteamConnection.GetPort() == -1)
+				{
+					UE_LOG_ONLINE(Log, TEXT("Closing all communications with user %s"), *SteamConnection.ToString(false));
+					ConnectionInfo.SteamNetworkingPtr->CloseP2PSessionWithUser(SteamConnection.SteamId);
+				}
+				else
+				{
+					UE_LOG_ONLINE(Log, TEXT("Closing channel %d with user %s"), SteamConnection.SteamChannel, *SteamConnection.ToString(false));
+					ConnectionInfo.SteamNetworkingPtr->CloseP2PChannelWithUser(SteamConnection.SteamId, SteamConnection.SteamChannel);
+					// If we no longer have any channels open with the user, we must remove the user, as Steam will do this automatically.
+					if (ConnectionInfo.ConnectedChannels.Num() != 0)
+					{
+						bShouldRemoveUser = false;
+						UE_LOG_ONLINE(Verbose, TEXT("%s still has %d open connections."), *SteamConnection.ToString(false), ConnectionInfo.ConnectedChannels.Num());
+					}
+					else
+					{
+						UE_LOG_ONLINE(Verbose, TEXT("%s has no more open connections! Going to remove"), *SteamConnection.ToString(false));
+					}
+				}
+
+				if (bShouldRemoveUser)
+				{
+					// Remove the user information from our current connections as they are no longer connected to us.
+					UE_LOG_ONLINE(Log, TEXT("%s has been removed."), *SteamConnection.ToString(false));
+					AcceptedConnections.Remove(SteamConnection.SteamId);
+				}
 			}
 
-			ExpiredSessions.Add(SessionId);
+			It.RemoveCurrent();
 		}
 	}
-
-	for (int32 SessionIdx=0; SessionIdx < ExpiredSessions.Num(); SessionIdx++)
-	{
-		FUniqueNetIdSteam& SessionId = ExpiredSessions[SessionIdx];
-		DeadConnections.Remove(SessionId);
-	}
 }
 
 /**
@@ -626,3 +736,22 @@ void FSocketSubsystemSteam::DumpSteamP2PSessionInfo(P2PSessionState_t& SessionIn
 	UE_LOG_ONLINE(Verbose, TEXT("-- QueuedBytes: %i, QueuedPackets: %i"), SessionInfo.m_nBytesQueuedForSend,
 		SessionInfo.m_nPacketsQueuedForSend);
 }
+
+/**
+ * Dumps all connection information for each user connection over SteamNet.
+ */
+void FSocketSubsystemSteam::DumpAllOpenSteamSessions()
+{
+	UE_LOG_ONLINE(Verbose, TEXT("Current Connection Info: "));
+	for (TMap<FUniqueNetIdSteam, FSteamP2PConnectionInfo>::TConstIterator It(AcceptedConnections); It; ++It)
+	{
+		UE_LOG_ONLINE(Verbose, TEXT("- Connection %s"), *It->Key.ToDebugString());
+		UE_LOG_ONLINE(Verbose, TEXT("--  Last Update Time: %d"), It->Value.LastReceivedTime);
+		FString ConnectedChannels(TEXT(""));
+		for (int32 i = 0; i < It->Value.ConnectedChannels.Num(); ++i)
+		{
+			ConnectedChannels = FString::Printf(TEXT("%s %d"), *ConnectedChannels, It->Value.ConnectedChannels[i]);
+		}
+		UE_LOG_ONLINE(Verbose, TEXT("--  Channels:%s"), *ConnectedChannels);
+	}
+}
diff --git a/Engine/Plugins/Online/OnlineSubsystemSteam/Source/Private/SocketSubsystemSteam.h b/Engine/Plugins/Online/OnlineSubsystemSteam/Source/Private/SocketSubsystemSteam.h
--- a/Engine/Plugins/Online/OnlineSubsystemSteam/Source/Private/SocketSubsystemSteam.h
+++ b/Engine/Plugins/Online/OnlineSubsystemSteam/Source/Private/SocketSubsystemSteam.h
@@ -6,6 +6,7 @@
 #include "UObject/WeakObjectPtr.h"
 #include "SocketSubsystem.h"
 #include "IPAddress.h"
+#include "IPAddressSteam.h"
 #include "OnlineSubsystemSteamTypes.h"
 #include "Containers/Ticker.h"
 #include "OnlineSubsystemSteamPackage.h"
@@ -15,7 +16,7 @@ class Error;
 /**
  * Windows specific socket subsystem implementation
  */
-class FSocketSubsystemSteam : public ISocketSubsystem, public FTickerObjectBase
+class FSocketSubsystemSteam : public ISocketSubsystem, public FTickerObjectBase, public FSelfRegisteringExec
 {
 protected:
 
@@ -28,35 +29,50 @@ protected:
 	/** Tracks existing Steamworks connections, for connection failure/timeout resolution */
 	TArray<struct FWeakObjectPtr> SteamConnections;
 
-    /** Keep track of Steam P2P connections */
+	/** Holds Steam connection information for each user */
 	struct FSteamP2PConnectionInfo
 	{
-        /** Steam networking interface responsible for this connection */
+		/** Steam networking interface responsible for this connection */
 		ISteamNetworking* SteamNetworkingPtr;
-		/** Channel to close */
-		int32 Channel;
-        /** Last time P2P session had activity (RecvFrom, etc) */
+
+		/** 
+		 * Last time the user's p2p session had activity (RecvFrom, etc). 
+		 * The value of this member is always the max value of the ConnectedChannels object
+		 */
 		double LastReceivedTime;
 
-		FSteamP2PConnectionInfo() :
-			SteamNetworkingPtr(NULL),
-			Channel(-1),
-			LastReceivedTime(0.0)
+		/** 
+		 * Channel connection ids for this user
+		 */
+		TArray<int32> ConnectedChannels;
+
+		FSteamP2PConnectionInfo(ISteamNetworking* InNetworkPtr=nullptr) :
+			SteamNetworkingPtr(InNetworkPtr),
+			LastReceivedTime(FPlatformTime::Seconds())
 		{
 		}
 
-		FSteamP2PConnectionInfo(ISteamNetworking* InSteamNetworkingPtr, double InTime, int32 InChannel=-1) :
-			SteamNetworkingPtr(InSteamNetworkingPtr),
-			Channel(InChannel),
-			LastReceivedTime(InTime)
+
+		/** This helper function automatically updates LastRecievedTime */
+		void AddOrUpdateChannel(int32 InChannel, double InTime)
 		{
+			ConnectedChannels.AddUnique(InChannel);
+			LastReceivedTime = FMath::Max(LastReceivedTime, InTime);
 		}
 	};
 
-    /** List of Steam P2P connections being tracked */
-	TMap<class FUniqueNetIdSteam, FSteamP2PConnectionInfo> AcceptedConnections;
-	/** List of Steam P2P connections to shutdown (dead connections remain around a few seconds longer to flush) */
-	TMap<class FUniqueNetIdSteam, FSteamP2PConnectionInfo> DeadConnections;
+    /** 
+	 * List of Steam P2P connections we have
+	 * As connections at start do not have a channel id, the key is just the accounts connected to us.
+	 */
+	TMap<FUniqueNetIdSteam, FSteamP2PConnectionInfo> AcceptedConnections;
+
+	/** 
+	 * List of Steam P2P connections to shutdown.
+	 * If the FInternetAddrSteam has a channel id of -1, all connections are dropped from the user.
+	 * Also tracked is the time in which the connection was marked to be removed (for linger purposes)
+	 */
+	TMap<FInternetAddrSteam, double> DeadConnections;
 
 	/**
 	 * Should Steam P2P sockets all fall back to Steam servers relay if a direct connection fails
@@ -64,7 +80,7 @@ protected:
 	 */
 	bool bAllowP2PPacketRelay;
 	/** 
-     * Timeout period for any P2P session
+     * Timeout (in seconds) period for any P2P session
 	 * read from [OnlineSubsystemSteam.P2PConnectionTimeout]
      * (should be at least as long as NetDriver::ConnectionTimeout) 
      */
@@ -73,6 +89,14 @@ protected:
 	double P2PDumpCounter;
 	/** Connection info output interval */
 	double P2PDumpInterval;
+	/**
+	 * The timeout (in seconds) between when a connection/channel is marked as destroyed
+	 * and when it's cleaned up. This allows for catching lingering messages
+	 * from other users.
+	 * If set to 0, all dead connections will be cleaned up each tick.
+	 * read from [OnlineSubsystemSteam.P2PCleanupTimeout]
+	 */
+	double P2PCleanupTimeout;
 
 	/**
 	 * Adds a steam socket for tracking
@@ -113,8 +137,10 @@ PACKAGE_SCOPE:
 	/**
 	 * Iterate through the pending dead connections and permanently remove any that have been around
 	 * long enough to flush their contents
+	 *
+	 * @param bSkipLinger skips the timeout reserved for lingering connection information to a race condition
 	 */
-	void CleanupDeadConnections();
+	void CleanupDeadConnections(bool bSkipLinger);
 
 	/**
 	 * Associate the game server steam id with any sockets that were created prior to successful login
@@ -142,7 +168,7 @@ PACKAGE_SCOPE:
 	 * 
 	 * @param RemoteId failure address
 	 */
-	void ConnectFailure(class FUniqueNetIdSteam& RemoteId);
+	void ConnectFailure(const FUniqueNetIdSteam& RemoteId);
 
 	/**
 	 * Potentially accept an incoming connection from a Steam P2P request
@@ -152,17 +178,18 @@ PACKAGE_SCOPE:
 	 * 
 	 * @return true if accepted, false otherwise
 	 */
-	bool AcceptP2PConnection(ISteamNetworking* SteamNetworkingPtr, FUniqueNetIdSteam& RemoteId);
+	bool AcceptP2PConnection(ISteamNetworking* SteamNetworkingPtr, const FUniqueNetIdSteam& RemoteId);
 
 	/**
 	 * Add/update a Steam P2P connection as being recently accessed
 	 *
 	 * @param SteamNetworkingPtr proper networking interface that this session is communicating on
 	 * @param SessionId P2P session recently heard from
+	 * @param ChannelId the channel id that the update happened on
      *
      * @return true if the connection is active, false if this is in the dead connections list
 	 */
-	bool P2PTouch(ISteamNetworking* SteamNetworkingPtr, FUniqueNetIdSteam& SessionId);
+	bool P2PTouch(ISteamNetworking* SteamNetworkingPtr, const FUniqueNetIdSteam& SessionId, int32 ChannelId = -1);
 
 	/**
 	 * Remove a Steam P2P session from tracking and close the connection
@@ -170,7 +197,18 @@ PACKAGE_SCOPE:
 	 * @param SessionId P2P session to remove
 	 * @param Channel channel to close, -1 to close all communication
 	 */
-	void P2PRemove(FUniqueNetIdSteam& SessionId, int32 Channel = -1);
+	void P2PRemove(const FUniqueNetIdSteam& SessionId, int32 Channel = -1);
+
+	/**
+	 * Checks to see if a Steam P2P Connection is pending close on the given channel.
+	 *
+	 * Before checking the given channel, this function checks if the session is marked for
+	 * global disconnection.
+	 * 
+	 * @param SessionId the user id tied to the session disconnection
+	 * @param Channel the communications channel id for the user if it exists
+	 */
+	bool IsConnectionPendingRemoval(const FUniqueNetIdSteam& SteamId, int32 Channel);
 
 	/**
 	 * Determines if the SocketSubsystemSteam should override the platform
@@ -186,6 +224,18 @@ PACKAGE_SCOPE:
 	 */
 	bool ShouldOverrideDefaultSubsystem() const;
 
+	/**
+	 * Dumps the Steam P2P networking information for a given session state
+	 *
+	 * @param SessionInfo struct from Steam call to GetP2PSessionState
+	 */
+	void DumpSteamP2PSessionInfo(P2PSessionState_t& SessionInfo);
+
+	/**
+	 * Dumps all connection information for each user connection over SteamNet.
+	 */
+	void DumpAllOpenSteamSessions();
+
 public:
 
 	FSocketSubsystemSteam() :
@@ -193,6 +243,7 @@ public:
 		P2PConnectionTimeout(45.0f),
 		P2PDumpCounter(0.0),
 		P2PDumpInterval(10.0),
+		P2PCleanupTimeout(1.5),
 		LastSocketError(0)
 	{
 	}
@@ -334,13 +385,6 @@ public:
 	 */
 	virtual TSharedRef<FInternetAddr> GetLocalBindAddr(FOutputDevice& Out) override;
 
-	/**
-	 * Dumps the Steam P2P networking information for a given session state
-	 *
-	 * @param SessionInfo struct from Steam call to GetP2PSessionState
-	 */
-	void DumpSteamP2PSessionInfo(P2PSessionState_t& SessionInfo);
-
 	/**
 	 * Chance for the socket subsystem to get some time
 	 *
@@ -352,6 +396,9 @@ public:
 	 * Waiting on a socket is not supported.
 	 */
 	virtual bool IsSocketWaitSupported() const override { return false; }
+
+	// FSelfRegisteringExec
+	virtual bool Exec(UWorld* InWorld, const TCHAR* Cmd, FOutputDevice& Ar) override;
 };
 
 /**
diff --git a/Engine/Plugins/Online/OnlineSubsystemSteam/Source/Private/SocketsSteam.cpp b/Engine/Plugins/Online/OnlineSubsystemSteam/Source/Private/SocketsSteam.cpp
--- a/Engine/Plugins/Online/OnlineSubsystemSteam/Source/Private/SocketsSteam.cpp
+++ b/Engine/Plugins/Online/OnlineSubsystemSteam/Source/Private/SocketsSteam.cpp
@@ -173,7 +173,7 @@ bool FSocketSteam::RecvFrom(uint8* Data, int32 BufferSize, int32& BytesRead, FIn
 	}
 	else
 	{
-		if (SocketSubsystem->P2PTouch( SteamNetworkingPtr, SteamAddr.SteamId))
+		if (SocketSubsystem->P2PTouch(SteamNetworkingPtr, SteamAddr.SteamId, SteamChannel))
 		{
 			SocketSubsystem->LastSocketError = SE_NO_ERROR;
 		}
