Fix so that steam OSS does not crash in development when it cannot access its steam_appid file.

diff --git a/Engine/Plugins/Online/OnlineSubsystemSteam/Source/Private/OnlineSubsystemSteam.cpp b/Engine/Plugins/Online/OnlineSubsystemSteam/Source/Private/OnlineSubsystemSteam.cpp
--- a/Engine/Plugins/Online/OnlineSubsystemSteam/Source/Private/OnlineSubsystemSteam.cpp
+++ b/Engine/Plugins/Online/OnlineSubsystemSteam/Source/Private/OnlineSubsystemSteam.cpp
@@ -124,7 +124,7 @@ inline FString GetSteamAppIdFilename()
  * Write out the steam app id to the steam_appid.txt file before initializing the API
  * @param SteamAppId id assigned to the application by Steam
  */
-static void WriteSteamAppIdToDisk(int32 SteamAppId)
+static bool WriteSteamAppIdToDisk(int32 SteamAppId)
 {
 	if (SteamAppId > 0)
 	{
@@ -136,7 +136,8 @@ static void WriteSteamAppIdToDisk(int32 SteamAppId)
 		IFileHandle* Handle = IPlatformFile::GetPlatformPhysical().OpenWrite(*SteamAppIdFilename, false, false);
 		if (!Handle)
 		{
-			UE_LOG_ONLINE(Fatal, TEXT("Failed to create file: %s"), *SteamAppIdFilename);
+			UE_LOG_ONLINE(Error, TEXT("Failed to create file: %s"), *SteamAppIdFilename);
+			return false;
 		}
 		else
 		{
@@ -148,8 +149,13 @@ static void WriteSteamAppIdToDisk(int32 SteamAppId)
 			Handle->Write(Archive.GetData(), Archive.Num());
 			delete Handle;
 			Handle = nullptr;
+
+			return true;
 		}
 	}
+
+	UE_LOG_ONLINE(Warning, TEXT("Steam App Id provided (%d) is invalid, must be greater than 0!"), SteamAppId);
+	return false;
 }
 
 /**
@@ -173,18 +179,25 @@ static void DeleteSteamAppIdFromDisk()
  *
  * @param RequireRelaunch enforce the Steam client running precondition
  * @param RelaunchAppId appid to launch when the Steam client is loaded
+ *
+ * @return if this sequence completed without any serious errors
  */
-void ConfigureSteamInitDevOptions(bool& RequireRelaunch, int32& RelaunchAppId)
+bool ConfigureSteamInitDevOptions(bool& RequireRelaunch, int32& RelaunchAppId)
 {
 #if !UE_BUILD_SHIPPING && !UE_BUILD_SHIPPING_WITH_EDITOR
 	// Write out the steam_appid.txt file before launching
 	if (!GConfig->GetInt(TEXT("OnlineSubsystemSteam"), TEXT("SteamDevAppId"), RelaunchAppId, GEngineIni))
 	{
 		UE_LOG_ONLINE(Warning, TEXT("Missing SteamDevAppId key in OnlineSubsystemSteam of DefaultEngine.ini"));
+		return false;
 	}
 	else
 	{
-		WriteSteamAppIdToDisk(RelaunchAppId);
+		if (!WriteSteamAppIdToDisk(RelaunchAppId))
+		{
+			UE_LOG_ONLINE(Warning, TEXT("Could not create/update the steam_appid.txt file! Make sure the directory is writable and there isn't another instance using this file"));
+			return false;
+		}
 	}
 
 	// Should the game force a relaunch in Steam if the client isn't already loaded
@@ -198,6 +211,8 @@ void ConfigureSteamInitDevOptions(bool& RequireRelaunch, int32& RelaunchAppId)
 	RequireRelaunch = true;
 	RelaunchAppId = UE4_PROJECT_STEAMSHIPPINGID;
 #endif
+
+	return true;
 }
 
 FOnlineAuthSteamPtr FOnlineSubsystemSteam::GetAuthInterface() const
@@ -381,7 +396,12 @@ bool FOnlineSubsystemSteam::Init()
 {
 	bool bRelaunchInSteam = false;
 	int RelaunchAppId = 0;
-	ConfigureSteamInitDevOptions(bRelaunchInSteam, RelaunchAppId);
+
+	if (!ConfigureSteamInitDevOptions(bRelaunchInSteam, RelaunchAppId))
+	{
+		UE_LOG_ONLINE(Warning, TEXT("Could not set up the steam environment! Falling back to another OSS."));
+		return false;
+	}
 
 	const bool bIsServer = IsRunningDedicatedServer();
 	bool bAttemptServerInit = true;
