Fix CreateSocket and CreateInternetAddr functional usage so they no longer use the deprecated methods.

diff --git a/Engine/Plugins/Experimental/BackChannel/Source/BackChannel/Private/Transport/BackChannelConnection.cpp b/Engine/Plugins/Experimental/BackChannel/Source/BackChannel/Private/Transport/BackChannelConnection.cpp
--- a/Engine/Plugins/Experimental/BackChannel/Source/BackChannel/Private/Transport/BackChannelConnection.cpp
+++ b/Engine/Plugins/Experimental/BackChannel/Source/BackChannel/Private/Transport/BackChannelConnection.cpp
@@ -176,7 +176,8 @@ bool FBackChannelConnection::Listen(const int16 Port)
     
     if (NewSocket == nullptr && SocketSubsystem != nullptr)
     {
-        NewSocket = SocketSubsystem->CreateSocket(NAME_Stream, TEXT("FBackChannelConnection ListenSocket"), true);
+		TSharedRef<FInternetAddr> BindAddress = Endpoint.ToInternetAddr();
+        NewSocket = SocketSubsystem->CreateSocket(NAME_Stream, TEXT("FBackChannelConnection ListenSocket"), BindAddress->GetProtocolType());
         
         if (NewSocket != nullptr)
         {
@@ -189,7 +190,7 @@ bool FBackChannelConnection::Listen(const int16 Port)
             
             if (!Error)
             {
-                Error = !NewSocket->Bind(*Endpoint.ToInternetAddr());
+                Error = !NewSocket->Bind(*BindAddress);
             }
             
             if (!Error)
diff --git a/Engine/Plugins/Online/OnlineSubsystem/Source/Private/LANBeacon.cpp b/Engine/Plugins/Online/OnlineSubsystem/Source/Private/LANBeacon.cpp
--- a/Engine/Plugins/Online/OnlineSubsystem/Source/Private/LANBeacon.cpp
+++ b/Engine/Plugins/Online/OnlineSubsystem/Source/Private/LANBeacon.cpp
@@ -39,17 +39,22 @@ bool FLanBeacon::Init(int32 Port)
 {
 	ISocketSubsystem* SocketSubsystem = ISocketSubsystem::Get(PLATFORM_SOCKETSUBSYSTEM);
 	bool bSuccess = false;
+	// Now the listen address
+	ListenAddr = SocketSubsystem->GetLocalBindAddr(*GWarn);
+	ListenAddr->SetPort(Port);
 	// Set our broadcast address
 	BroadcastAddr = SocketSubsystem->CreateInternetAddr();
 	BroadcastAddr->SetBroadcastAddress();
 	BroadcastAddr->SetPort(Port);
-	// Now the listen address
-	ListenAddr = SocketSubsystem->GetLocalBindAddr(*GWarn);
-	ListenAddr->SetPort(Port);
 	// A temporary "received from" address
 	SockAddr = SocketSubsystem->CreateInternetAddr();
+
 	// Now create and set up our sockets (no VDP)
-	ListenSocket = SocketSubsystem->CreateSocket(NAME_DGram, TEXT("LAN beacon"), true);
+	// TODO: @jleonard - for the time being this will be fine, but on hybrid stacks this will be an issue as 
+	// we have no guarantee that the address protocol from GetLocalBindAddr and BroadcastAddr will match.
+	PRAGMA_DISABLE_DEPRECATION_WARNINGS
+	ListenSocket = SocketSubsystem->CreateSocket(NAME_DGram, TEXT("LAN beacon"), ListenAddr->GetProtocolType());
+	PRAGMA_ENABLE_DEPRECATION_WARNINGS
 	if (ListenSocket != NULL)
 	{
 		ListenSocket->SetReuseAddr();
diff --git a/Engine/Plugins/Online/OnlineSubsystemSteam/Source/Private/OnlineSessionAsyncServerSteam.cpp b/Engine/Plugins/Online/OnlineSubsystemSteam/Source/Private/OnlineSessionAsyncServerSteam.cpp
--- a/Engine/Plugins/Online/OnlineSubsystemSteam/Source/Private/OnlineSessionAsyncServerSteam.cpp
+++ b/Engine/Plugins/Online/OnlineSubsystemSteam/Source/Private/OnlineSessionAsyncServerSteam.cpp
@@ -406,7 +406,9 @@ void FOnlineAsyncTaskSteamCreateServer::Finalize()
 			UE_LOG_ONLINE_SESSION(Verbose, TEXT("Server SteamP2P IP: %s"), *NewSessionInfo->SteamP2PAddr->ToString(true));
 
 			// Create the proper ip address for this server
-			NewSessionInfo->HostAddr = ISocketSubsystem::Get(PLATFORM_SOCKETSUBSYSTEM)->CreateInternetAddr(SteamGameServerPtr->GetPublicIP(), Subsystem->GetGameServerGamePort());
+			NewSessionInfo->HostAddr = ISocketSubsystem::Get(PLATFORM_SOCKETSUBSYSTEM)->CreateInternetAddr();
+			NewSessionInfo->HostAddr->SetIp(SteamGameServerPtr->GetPublicIP());
+			NewSessionInfo->HostAddr->SetPort(Subsystem->GetGameServerGamePort());
 			UE_LOG_ONLINE_SESSION(Verbose, TEXT("Server IP: %s"), *NewSessionInfo->HostAddr->ToString(true));
 
 			if (!Session->OwningUserId.IsValid())
diff --git a/Engine/Plugins/Online/OnlineSubsystemSteam/Source/Private/OnlineSessionInterfaceSteam.cpp b/Engine/Plugins/Online/OnlineSubsystemSteam/Source/Private/OnlineSessionInterfaceSteam.cpp
--- a/Engine/Plugins/Online/OnlineSubsystemSteam/Source/Private/OnlineSessionInterfaceSteam.cpp
+++ b/Engine/Plugins/Online/OnlineSubsystemSteam/Source/Private/OnlineSessionInterfaceSteam.cpp
@@ -999,9 +999,7 @@ uint32 FOnlineSessionSteam::JoinLANSession(int32 PlayerNum, FNamedOnlineSession*
 		const FOnlineSessionInfoSteam* SearchSessionInfo = (const FOnlineSessionInfoSteam*)SearchSession->SessionInfo.Get();
 		FOnlineSessionInfoSteam* SessionInfo = (FOnlineSessionInfoSteam*)Session->SessionInfo.Get();
 
-		uint32 IpAddr;
-		SearchSessionInfo->HostAddr->GetIp(IpAddr);
-		SessionInfo->HostAddr = ISocketSubsystem::Get(PLATFORM_SOCKETSUBSYSTEM)->CreateInternetAddr(IpAddr, SearchSessionInfo->HostAddr->GetPort());
+		SessionInfo->HostAddr = SearchSessionInfo->HostAddr->Clone();
 		Result = ONLINE_SUCCESS;
 	}
 
@@ -1038,7 +1036,9 @@ bool FOnlineSessionSteam::FindFriendSession(int32 LocalUserNum, const FUniqueNet
 				else
 				{
 					// Search for the session via host ip
-					TSharedRef<FInternetAddr> IpAddr = ISocketSubsystem::Get()->CreateInternetAddr(FriendGameInfo.m_unGameIP, FriendGameInfo.m_usGamePort);
+					TSharedRef<FInternetAddr> IpAddr = ISocketSubsystem::Get()->CreateInternetAddr();
+					IpAddr->SetIp(FriendGameInfo.m_unGameIP);
+					IpAddr->SetPort(FriendGameInfo.m_usGamePort);
 					CurrentSessionSearch->QuerySettings.Set(FName(SEARCH_STEAM_HOSTIP), IpAddr->ToString(true), EOnlineComparisonOp::Equals);
 
 					FOnlineAsyncTaskSteamFindServerForFriendSession* NewTask = new FOnlineAsyncTaskSteamFindServerForFriendSession(SteamSubsystem, CurrentSessionSearch, LocalUserNum, OnFindFriendSessionCompleteDelegates[LocalUserNum]);
diff --git a/Engine/Plugins/Online/OnlineSubsystemSteam/Source/Private/SocketSubsystemSteam.cpp b/Engine/Plugins/Online/OnlineSubsystemSteam/Source/Private/SocketSubsystemSteam.cpp
--- a/Engine/Plugins/Online/OnlineSubsystemSteam/Source/Private/SocketSubsystemSteam.cpp
+++ b/Engine/Plugins/Online/OnlineSubsystemSteam/Source/Private/SocketSubsystemSteam.cpp
@@ -263,10 +263,10 @@ void FSocketSubsystemSteam::RegisterConnection(USteamNetConnection* Connection)
 	FWeakObjectPtr ObjectPtr = Connection;
 	SteamConnections.Add(ObjectPtr);
 
-	if (Connection->GetInternetAddr().IsValid() && Connection->Socket)
+	if (Connection->GetRemoteAddr().IsValid() && Connection->Socket)
 	{
 		FSocketSteam* SteamSocket = (FSocketSteam*)Connection->Socket;
-		TSharedPtr<const FInternetAddrSteam> SteamAddr = StaticCastSharedPtr<const FInternetAddrSteam>(Connection->GetInternetAddr());
+		TSharedPtr<const FInternetAddrSteam> SteamAddr = StaticCastSharedPtr<const FInternetAddrSteam>(Connection->GetRemoteAddr());
 
 		UE_LOG_ONLINE(Log, TEXT("Adding user %s from RegisterConnection"), *SteamAddr->ToString(true));
 		P2PTouch(SteamSocket->SteamNetworkingPtr, SteamAddr->SteamId, SteamAddr->SteamChannel);
@@ -289,9 +289,9 @@ void FSocketSubsystemSteam::UnregisterConnection(USteamNetConnection* Connection
 	// is garbage collected. It's possible that the player who left rejoined before garbage
 	// collection runs (their connection object will be different), so P2PRemove would kick
 	// them from the session when it shouldn't.
-	if (SteamConnections.RemoveSingleSwap(ObjectPtr) == 1 && Connection->GetInternetAddr().IsValid())
+	if (SteamConnections.RemoveSingleSwap(ObjectPtr) == 1 && Connection->GetRemoteAddr().IsValid())
 	{
-		TSharedPtr<const FInternetAddrSteam> SteamAddr = StaticCastSharedPtr<const FInternetAddrSteam>(Connection->GetInternetAddr());
+		TSharedPtr<const FInternetAddrSteam> SteamAddr = StaticCastSharedPtr<const FInternetAddrSteam>(Connection->GetRemoteAddr());
 		P2PRemove(SteamAddr->SteamId, SteamAddr->SteamChannel);
 	}
 }
@@ -311,7 +311,7 @@ void FSocketSubsystemSteam::ConnectFailure(const FUniqueNetIdSteam& RemoteId)
 	for (int32 ConnIdx=0; ConnIdx<SteamConnections.Num(); ConnIdx++)
 	{
 		USteamNetConnection* SteamConn = CastChecked<USteamNetConnection>(SteamConnections[ConnIdx].Get());
-		TSharedPtr<const FInternetAddrSteam> RemoteAddrSteam = StaticCastSharedPtr<const FInternetAddrSteam>(SteamConn->GetInternetAddr());
+		TSharedPtr<const FInternetAddrSteam> RemoteAddrSteam = StaticCastSharedPtr<const FInternetAddrSteam>(SteamConn->GetRemoteAddr());
 
 		// Only checking Id here because its a complete failure (channel doesn't matter)
 		if (RemoteAddrSteam->SteamId == RemoteId)
@@ -791,7 +791,9 @@ void FSocketSubsystemSteam::DumpSteamP2PSessionInfo(P2PSessionState_t& SessionIn
 		return;
 	}
 
-	TSharedRef<FInternetAddr> IpAddr = ISocketSubsystem::Get(PLATFORM_SOCKETSUBSYSTEM)->CreateInternetAddr(SessionInfo.m_nRemoteIP, SessionInfo.m_nRemotePort);
+	TSharedRef<FInternetAddr> IpAddr = ISocketSubsystem::Get(PLATFORM_SOCKETSUBSYSTEM)->CreateInternetAddr();
+	IpAddr->SetIp(SessionInfo.m_nRemoteIP);
+	IpAddr->SetPort(SessionInfo.m_nRemotePort);
 	UE_LOG_ONLINE(Verbose, TEXT("- Detailed P2P session info:"));
 	UE_LOG_ONLINE(Verbose, TEXT("-- IPAddress: %s"), *IpAddr->ToString(true));
 	UE_LOG_ONLINE(Verbose, TEXT("-- ConnectionActive: %i, Connecting: %i, SessionError: %i, UsingRelay: %i"),
diff --git a/Engine/Plugins/Runtime/AR/Apple/AppleARKitFaceSupport/Source/AppleARKitFaceSupport/Private/AppleARKitLiveLinkSource.cpp b/Engine/Plugins/Runtime/AR/Apple/AppleARKitFaceSupport/Source/AppleARKitFaceSupport/Private/AppleARKitLiveLinkSource.cpp
--- a/Engine/Plugins/Runtime/AR/Apple/AppleARKitFaceSupport/Source/AppleARKitFaceSupport/Private/AppleARKitLiveLinkSource.cpp
+++ b/Engine/Plugins/Runtime/AR/Apple/AppleARKitFaceSupport/Source/AppleARKitFaceSupport/Private/AppleARKitLiveLinkSource.cpp
@@ -248,7 +248,7 @@ bool FAppleARKitLiveLinkRemotePublisher::InitSendSocket()
 	{
 		ISocketSubsystem* SocketSubsystem = ISocketSubsystem::Get();
 		// Allocate our socket for sending
-		SendSocket = SocketSubsystem->CreateSocket(NAME_DGram, TEXT("FAppleARKitLiveLinkRemotePublisher socket"), true);
+		SendSocket = SocketSubsystem->CreateSocket(NAME_DGram, TEXT("FAppleARKitLiveLinkRemotePublisher socket"), Addr->GetProtocolType());
 		SendSocket->SetReuseAddr();
 		SendSocket->SetNonBlocking();
 		UE_LOG(LogAppleARKitFace, Log, TEXT("Sending LiveLink face AR data to address (%s)"), *Addr->ToString(true));
@@ -339,7 +339,7 @@ bool FAppleARKitLiveLinkRemoteListener::InitReceiveSocket()
 	GConfig->GetInt(TEXT("/Script/AppleARKit.AppleARKitSettings"), TEXT("LiveLinkPublishingPort"), LiveLinkPort, GEngineIni);
 	Addr->SetPort(LiveLinkPort);
 
-	RecvSocket = SocketSubsystem->CreateSocket(NAME_DGram, TEXT("FAppleARKitLiveLinkRemoteListener socket"));
+	RecvSocket = SocketSubsystem->CreateSocket(NAME_DGram, TEXT("FAppleARKitLiveLinkRemoteListener socket"), Addr->GetProtocolType());
 	if (RecvSocket != nullptr)
 	{
 		RecvSocket->SetReuseAddr();
diff --git a/Engine/Plugins/Runtime/Database/RemoteDatabaseSupport/Source/RemoteDatabaseSupport/Private/RemoteDatabaseConnection.cpp b/Engine/Plugins/Runtime/Database/RemoteDatabaseSupport/Source/RemoteDatabaseSupport/Private/RemoteDatabaseConnection.cpp
--- a/Engine/Plugins/Runtime/Database/RemoteDatabaseSupport/Source/RemoteDatabaseSupport/Private/RemoteDatabaseConnection.cpp
+++ b/Engine/Plugins/Runtime/Database/RemoteDatabaseSupport/Source/RemoteDatabaseSupport/Private/RemoteDatabaseConnection.cpp
@@ -35,16 +35,8 @@ bool ExecuteDBProxyCommand(FSocket *Socket, const FString& Cmd)
  * Constructor.
  */
 FRemoteDatabaseConnection::FRemoteDatabaseConnection()
-: Socket(NULL)
+: Socket(nullptr)
 {
-	ISocketSubsystem* SocketSubsystem = ISocketSubsystem::Get(PLATFORM_SOCKETSUBSYSTEM);
-	check(SocketSubsystem);
-
-	// The socket won't work if secure connections are enabled, so don't try
-	if (SocketSubsystem->RequiresEncryptedPackets() == false)
-	{
-		Socket = SocketSubsystem->CreateSocket(NAME_Stream, TEXT("remote database connection"));
-	}
 }
 
 /**
@@ -58,7 +50,7 @@ FRemoteDatabaseConnection::~FRemoteDatabaseConnection()
 	if (Socket)
 	{
 		SocketSubsystem->DestroySocket(Socket);
-		Socket = NULL;
+		Socket = nullptr;
 	}
 }
 
@@ -74,23 +66,30 @@ FRemoteDatabaseConnection::~FRemoteDatabaseConnection()
 bool FRemoteDatabaseConnection::Open( const TCHAR* ConnectionString, const TCHAR* RemoteConnectionIP, const TCHAR* RemoteConnectionStringOverride )
 {
 	bool bIsValid = false;
-	if ( Socket )
+	ISocketSubsystem* SocketSubsystem = ISocketSubsystem::Get(PLATFORM_SOCKETSUBSYSTEM);
+	check(SocketSubsystem);
+
+	TSharedRef<FInternetAddr> Address = SocketSubsystem->CreateInternetAddr();
+	Address->SetIp(RemoteConnectionIP, bIsValid);
+	Address->SetPort(10500);
+
+	// The socket won't work if secure connections are enabled, so don't try
+	if (SocketSubsystem->RequiresEncryptedPackets() == false && Socket == nullptr)
+	{
+		Socket = SocketSubsystem->CreateSocket(NAME_Stream, TEXT("remote database connection"), Address->GetProtocolType());
+	}
+	else
 	{
-		ISocketSubsystem* SocketSubsystem = ISocketSubsystem::Get(PLATFORM_SOCKETSUBSYSTEM);
-		check(SocketSubsystem);
+		bIsValid = false;
+	}
 
-		TSharedRef<FInternetAddr> Address = SocketSubsystem->CreateInternetAddr();
-		Address->SetIp(RemoteConnectionIP, bIsValid);
-		Address->SetPort(10500);
+	if (bIsValid && Socket)
+	{
+		bIsValid = Socket->Connect(*Address);
 
-		if(bIsValid)
+		if (bIsValid && RemoteConnectionStringOverride)
 		{
-			bIsValid = Socket->Connect(*Address);
-
-			if(bIsValid && RemoteConnectionStringOverride)
-			{
-				SetConnectionString(RemoteConnectionStringOverride);
-			}
+			SetConnectionString(RemoteConnectionStringOverride);
 		}
 	}
 	return bIsValid;
@@ -274,7 +273,7 @@ float FRemoteDataBaseRecordSet::GetFloat( const TCHAR* Column ) const
 }
 
 /** Constructor. */
-FRemoteDataBaseRecordSet::FRemoteDataBaseRecordSet(int32 ResultSetID, FSocket *Connection) : Socket(NULL)
+FRemoteDataBaseRecordSet::FRemoteDataBaseRecordSet(int32 ResultSetID, FSocket *Connection) : Socket(nullptr)
 {
 	check(ResultSetID >= 0);
 	check(Connection);
diff --git a/Engine/Source/Editor/SwarmInterface/Private/SwarmInterfaceLocal.cpp b/Engine/Source/Editor/SwarmInterface/Private/SwarmInterfaceLocal.cpp
--- a/Engine/Source/Editor/SwarmInterface/Private/SwarmInterfaceLocal.cpp
+++ b/Engine/Source/Editor/SwarmInterface/Private/SwarmInterfaceLocal.cpp
@@ -138,10 +138,11 @@ namespace SwarmInterfaceLocalImpl
 		if (SocketSubsystem)
 		{
 			// create socket
-			FSocket* Socket = SocketSubsystem->CreateSocket(NAME_DGram, TEXT("TestSocket"), true);
+			TSharedRef<FInternetAddr> BindAddr = FIPv4Endpoint::Any.ToInternetAddr();
+			FSocket* Socket = SocketSubsystem->CreateSocket(NAME_DGram, TEXT("TestSocket"), BindAddr->GetProtocolType());
 			if (Socket)
 			{
-				if (Socket->Bind(*FIPv4Endpoint::Any.ToInternetAddr()))
+				if (Socket->Bind(*BindAddr))
 				{
 					if (Socket->SetBroadcast(true) && Socket->SetMulticastLoopback(true))
 					{
diff --git a/Engine/Source/Runtime/NetworkFile/Private/TCPTransport.cpp b/Engine/Source/Runtime/NetworkFile/Private/TCPTransport.cpp
--- a/Engine/Source/Runtime/NetworkFile/Private/TCPTransport.cpp
+++ b/Engine/Source/Runtime/NetworkFile/Private/TCPTransport.cpp
@@ -26,15 +26,16 @@ bool FTCPTransport::Initialize(const TCHAR* InHostIp)
 
 	// convert the string to a ip addr structure
 	// DEFAULT_TCP_FILE_SERVING_PORT is overridden 
-	TSharedRef<FInternetAddr> Addr = SSS->CreateInternetAddr(0, DEFAULT_TCP_FILE_SERVING_PORT);
+	TSharedRef<FInternetAddr> Addr = SSS->CreateInternetAddr();
 	bool bIsValid;
 
 	Addr->SetIp(*HostIp, bIsValid);
+	Addr->SetPort(DEFAULT_TCP_FILE_SERVING_PORT);
 
 	if (bIsValid)
 	{
 		// create the socket
-		FileSocket = SSS->CreateSocket(NAME_Stream, TEXT("FNetworkPlatformFile tcp"));
+		FileSocket = SSS->CreateSocket(NAME_Stream, TEXT("FNetworkPlatformFile tcp"), Addr->GetProtocolType());
 
 		// try to connect to the server
 		if (FileSocket == nullptr || FileSocket->Connect(*Addr) == false)
diff --git a/Engine/Source/Runtime/NetworkFileSystem/Private/NetworkFileServer.cpp b/Engine/Source/Runtime/NetworkFileSystem/Private/NetworkFileServer.cpp
--- a/Engine/Source/Runtime/NetworkFileSystem/Private/NetworkFileServer.cpp
+++ b/Engine/Source/Runtime/NetworkFileSystem/Private/NetworkFileServer.cpp
@@ -136,8 +136,12 @@ FNetworkFileServer::FNetworkFileServer( int32 InPort, FNetworkFileDelegateContai
 	}
 	else
 	{
+		// listen on any IP address
+		ListenAddr = SocketSubsystem->GetLocalBindAddr(*GLog);
+		ListenAddr->SetPort(InPort);
+
 		// create a server TCP socket
-		Socket = SocketSubsystem->CreateSocket(NAME_Stream, TEXT("FNetworkFileServer tcp-listen"));
+		Socket = SocketSubsystem->CreateSocket(NAME_Stream, TEXT("FNetworkFileServer tcp-listen"), ListenAddr->GetProtocolType());
 
 		if(!Socket)
 		{
@@ -145,10 +149,6 @@ FNetworkFileServer::FNetworkFileServer( int32 InPort, FNetworkFileDelegateContai
 		}
 		else
 		{
-			// listen on any IP address
-			ListenAddr = SocketSubsystem->GetLocalBindAddr(*GLog);
-
-			ListenAddr->SetPort(InPort);
 			Socket->SetReuseAddr();
 
 			// bind to the address
diff --git a/Engine/Source/Runtime/Networking/Public/Common/TcpSocketBuilder.h b/Engine/Source/Runtime/Networking/Public/Common/TcpSocketBuilder.h
--- a/Engine/Source/Runtime/Networking/Public/Common/TcpSocketBuilder.h
+++ b/Engine/Source/Runtime/Networking/Public/Common/TcpSocketBuilder.h
@@ -8,6 +8,7 @@
 #include "Interfaces/IPv4/IPv4Address.h"
 #include "Interfaces/IPv4/IPv4Endpoint.h"
 #include "Sockets.h"
+#include "SocketTypes.h"
 
 class Error;
 
@@ -215,7 +216,7 @@ public:
 
 		if (SocketSubsystem != nullptr)
 		{
-			Socket = SocketSubsystem->CreateSocket(NAME_Stream, *Description, ESocketProtocolFamily::IPv4);
+			Socket = SocketSubsystem->CreateSocket(NAME_Stream, *Description, FNetworkProtocolTypes::IPv4);
 
 			if (Socket != nullptr)
 			{
diff --git a/Engine/Source/Runtime/Networking/Public/Common/UdpSocketBuilder.h b/Engine/Source/Runtime/Networking/Public/Common/UdpSocketBuilder.h
--- a/Engine/Source/Runtime/Networking/Public/Common/UdpSocketBuilder.h
+++ b/Engine/Source/Runtime/Networking/Public/Common/UdpSocketBuilder.h
@@ -275,14 +275,15 @@ public:
 		}
 
 		// create socket
-		FSocket* Socket = SocketSubsystem->CreateSocket(NAME_DGram, *Description, true);
+		TSharedRef<FInternetAddr> RemoteAddr = BoundEndpoint.ToInternetAddr();
+		FSocket* Socket = SocketSubsystem->CreateSocket(NAME_DGram, *Description, RemoteAddr->GetProtocolType());
 
 		if (Socket == nullptr)
 		{
 			GLog->Logf(TEXT("FUdpSocketBuilder: Failed to create socket %s"), *Description);
 			return nullptr;
 		}
-
+		
 		// configure socket
 		bool Error =
 			!Socket->SetNonBlocking(!Blocking) ||
@@ -297,7 +298,7 @@ public:
 		}
 		else
 		{
-			Error = Bound && !Socket->Bind(*BoundEndpoint.ToInternetAddr());
+			Error = Bound && !Socket->Bind(*RemoteAddr);
 		}
 
 		// configure multicast
diff --git a/Engine/Source/Runtime/Online/ICMP/Private/Icmp.cpp b/Engine/Source/Runtime/Online/ICMP/Private/Icmp.cpp
--- a/Engine/Source/Runtime/Online/ICMP/Private/Icmp.cpp
+++ b/Engine/Source/Runtime/Online/ICMP/Private/Icmp.cpp
@@ -53,7 +53,9 @@ bool ResolveIp(ISocketSubsystem* SocketSub, const FString& HostName, FString& Ou
 	if (SocketSub)
 	{
 		TSharedRef<FInternetAddr> HostAddr = SocketSub->CreateInternetAddr();
+		PRAGMA_DISABLE_DEPRECATION_WARNINGS
 		ESocketErrors HostResolveError = SocketSub->GetHostByName(TCHAR_TO_ANSI(*HostName), *HostAddr);
+		PRAGMA_ENABLE_DEPRECATION_WARNINGS
 		if (HostResolveError == SE_NO_ERROR || HostResolveError == SE_EWOULDBLOCK)
 		{
 			OutIp = HostAddr->ToString(false);
diff --git a/Engine/Source/Runtime/Online/ICMP/Private/UDPPing.cpp b/Engine/Source/Runtime/Online/ICMP/Private/UDPPing.cpp
--- a/Engine/Source/Runtime/Online/ICMP/Private/UDPPing.cpp
+++ b/Engine/Source/Runtime/Online/ICMP/Private/UDPPing.cpp
@@ -81,21 +81,21 @@ FIcmpEchoResult UDPEchoImpl(ISocketSubsystem* SocketSub, const FString& TargetAd
 	//ISocketSubsystem* SocketSub = ISocketSubsystem::Get(PLATFORM_SOCKETSUBSYSTEM);
 	if (SocketSub)
 	{
-		FSocket* Socket = SocketSub->CreateSocket(NAME_DGram, TEXT("UDPPing"), false);
-		if (Socket)
+		TSharedRef<FInternetAddr> ToAddr = SocketSub->CreateInternetAddr();
+
+		bool bIsValid = false;
+		ToAddr->SetIp(*ResolvedAddress, bIsValid);
+		ToAddr->SetPort(Port);
+		
+		if (bIsValid)
 		{
 			uint8 SendBuffer[PacketSize];
 			// Clear packet buffer
 			FMemory::Memset(SendBuffer, 0, PacketSize);
-
 			Result.ResolvedAddress = ResolvedAddress;
 
-			TSharedRef<FInternetAddr> ToAddr = SocketSub->CreateInternetAddr();
-
-			bool bIsValid = false;
-			ToAddr->SetIp(*ResolvedAddress, bIsValid);
-			ToAddr->SetPort(Port);
-			if (bIsValid)
+			FSocket* Socket = SocketSub->CreateSocket(NAME_DGram, TEXT("UDPPing"), ToAddr->GetProtocolType());
+			if (Socket)
 			{
 				uint16 SentId = (uint16)FPlatformProcess::GetCurrentProcessId();
 				uint16 SentSeq = PingSequence++;
diff --git a/Engine/Source/Runtime/PerfCounters/Private/PerfCounters.cpp b/Engine/Source/Runtime/PerfCounters/Private/PerfCounters.cpp
--- a/Engine/Source/Runtime/PerfCounters/Private/PerfCounters.cpp
+++ b/Engine/Source/Runtime/PerfCounters/Private/PerfCounters.cpp
@@ -73,8 +73,13 @@ bool FPerfCounters::Initialize()
 
 	ScratchIPAddr = SocketSubsystem->CreateInternetAddr();
 
+	// create a localhost binding for the requested port
+	TSharedRef<FInternetAddr> LocalhostAddr = SocketSubsystem->CreateInternetAddr();
+	LocalhostAddr->SetLoopbackAddress();
+	LocalhostAddr->SetPort(StatsPort);
+
 	// make our listen socket
-	Socket = SocketSubsystem->CreateSocket(NAME_Stream, TEXT("FPerfCounters"));
+	Socket = SocketSubsystem->CreateSocket(NAME_Stream, TEXT("FPerfCounters"), LocalhostAddr->GetProtocolType());
 	if (Socket == nullptr)
 	{
 		UE_LOG(LogPerfCounters, Error, TEXT("FPerfCounters unable to allocate stream socket"));
@@ -84,8 +89,6 @@ bool FPerfCounters::Initialize()
 	// make us non blocking
 	Socket->SetNonBlocking(true);
 
-	// create a localhost binding for the requested port
-	TSharedRef<FInternetAddr> LocalhostAddr = SocketSubsystem->CreateInternetAddr(0x7f000001 /* 127.0.0.1 */, StatsPort);
 	if (!Socket->Bind(*LocalhostAddr))
 	{
 		UE_LOG(LogPerfCounters, Error, TEXT("FPerfCounters unable to bind to %s"), *LocalhostAddr->ToString(true));
diff --git a/Engine/Source/Runtime/Sockets/Private/BSDIPv6Sockets/IPAddressBSDIPv6.cpp b/Engine/Source/Runtime/Sockets/Private/BSDIPv6Sockets/IPAddressBSDIPv6.cpp
--- a/Engine/Source/Runtime/Sockets/Private/BSDIPv6Sockets/IPAddressBSDIPv6.cpp
+++ b/Engine/Source/Runtime/Sockets/Private/BSDIPv6Sockets/IPAddressBSDIPv6.cpp
@@ -336,7 +336,9 @@ FResolveInfoCachedBSDIPv6::FResolveInfoCachedBSDIPv6(const FInternetAddr& InAddr
 	{
 		uint32 IpAddr;
 		InAddr.GetIp(IpAddr);
-		Addr = ISocketSubsystem::Get()->CreateInternetAddr(IpAddr, InAddr.GetPort());
+		Addr = ISocketSubsystem::Get()->CreateInternetAddr();
+		Addr->SetIp(IpAddr);
+		Addr->SetPort(InAddr.GetPort());
 	}
 	
 }
diff --git a/Engine/Source/Runtime/Sockets/Public/AddressInfoTypes.h b/Engine/Source/Runtime/Sockets/Public/AddressInfoTypes.h
--- a/Engine/Source/Runtime/Sockets/Public/AddressInfoTypes.h
+++ b/Engine/Source/Runtime/Sockets/Public/AddressInfoTypes.h
@@ -40,9 +40,10 @@ ENUM_CLASS_FLAGS(EAddressInfoFlags);
 
 struct FAddressInfoResultData
 {
-	UE_DEPRECATED(4.23, "Use constructor that supports an FName as the protocol specification as it supports multiple protocol types")
+	UE_DEPRECATED(4.23, "Use constructor that supports an FName as the AddressProtocol property is now deprecated. AddressProtocolName supports multiple protocol types as should be used instead.")
 	FAddressInfoResultData(TSharedRef<FInternetAddr> InAddr, SIZE_T InAddrLen, ESocketProtocolFamily InProtocol, ESocketType InSocketConfiguration) :
 		AddressProtocol(InProtocol),
+		AddressProtocolName(NAME_None),
 		SocketConfiguration(InSocketConfiguration),
 		AddressLen(InAddrLen),
 		Address(InAddr)
@@ -83,7 +84,6 @@ struct FAddressInfoResultData
 	}
 
 	/* The protocol of the address stored */
-	UE_DEPRECATED(4.23, "Use AddressProtocolName as it supports multiple protocol types")
 	ESocketProtocolFamily AddressProtocol;
 	/* The protocol name of the address stored */
 	FName AddressProtocolName;
