Check to see if a Steam session has been properly created before attempting to modify/update it. Fixes a crash that occurs.

diff --git a/Engine/Plugins/Online/OnlineSubsystemSteam/Source/Private/OnlineSessionInterfaceSteam.cpp b/Engine/Plugins/Online/OnlineSubsystemSteam/Source/Private/OnlineSessionInterfaceSteam.cpp
--- a/Engine/Plugins/Online/OnlineSubsystemSteam/Source/Private/OnlineSessionInterfaceSteam.cpp
+++ b/Engine/Plugins/Online/OnlineSubsystemSteam/Source/Private/OnlineSessionInterfaceSteam.cpp
@@ -444,18 +444,24 @@ bool FOnlineSessionSteam::UpdateSession(FName SessionName, FOnlineSessionSetting
 		if (!Session->SessionSettings.bIsLANMatch)
 		{
 			FOnlineSessionInfoSteam* SessionInfo = (FOnlineSessionInfoSteam*)(Session->SessionInfo.Get());
-
-			if (SessionInfo->SessionType == ESteamSession::LobbySession && SessionInfo->SessionId.IsValid())
+			if (SessionInfo)
 			{
-				// Lobby update
-				FOnlineAsyncTaskSteamUpdateLobby* NewTask = new FOnlineAsyncTaskSteamUpdateLobby(SteamSubsystem, SessionName, bShouldRefreshOnlineData, UpdatedSessionSettings);
-				SteamSubsystem->QueueAsyncTask(NewTask);
+				if (SessionInfo->SessionType == ESteamSession::LobbySession && SessionInfo->SessionId.IsValid())
+				{
+					// Lobby update
+					FOnlineAsyncTaskSteamUpdateLobby* NewTask = new FOnlineAsyncTaskSteamUpdateLobby(SteamSubsystem, SessionName, bShouldRefreshOnlineData, UpdatedSessionSettings);
+					SteamSubsystem->QueueAsyncTask(NewTask);
+				}
+				else if (SessionInfo->SessionType == ESteamSession::AdvertisedSessionHost)
+				{
+					// Gameserver update
+					FOnlineAsyncTaskSteamUpdateServer* NewTask = new FOnlineAsyncTaskSteamUpdateServer(SteamSubsystem, SessionName, bShouldRefreshOnlineData, UpdatedSessionSettings);
+					SteamSubsystem->QueueAsyncTask(NewTask);
+				}
 			}
-			else if (SessionInfo->SessionType == ESteamSession::AdvertisedSessionHost)
+			else
 			{
-				// Gameserver update
-				FOnlineAsyncTaskSteamUpdateServer* NewTask = new FOnlineAsyncTaskSteamUpdateServer(SteamSubsystem, SessionName, bShouldRefreshOnlineData, UpdatedSessionSettings);
-				SteamSubsystem->QueueAsyncTask(NewTask);
+				bWasSuccessful = false;
 			}
 		}
 		else
