Adding the missing nullchecks from VoiceEngineImpl to the Steam VoiceEngine. This prevents a crash on joining when using Steam voice.

diff --git a/Engine/Plugins/Online/OnlineSubsystemSteam/Source/Private/VoiceEngineSteam.cpp b/Engine/Plugins/Online/OnlineSubsystemSteam/Source/Private/VoiceEngineSteam.cpp
--- a/Engine/Plugins/Online/OnlineSubsystemSteam/Source/Private/VoiceEngineSteam.cpp
+++ b/Engine/Plugins/Online/OnlineSubsystemSteam/Source/Private/VoiceEngineSteam.cpp
@@ -520,7 +520,7 @@ void FVoiceEngineSteam::TickTalkers(float DeltaTime)
 		FRemoteTalkerDataSteam& RemoteData = It.Value();
 		double TimeSince = CurTime - RemoteData.LastSeen;
 
-		if (RemoteData.VoipSynthComponent->IsIdling() && RemoteData.bIsActive)
+		if (RemoteData.VoipSynthComponent && RemoteData.VoipSynthComponent->IsIdling() && RemoteData.bIsActive)
 		{
 			RemoteData.VoipSynthComponent->Stop();
 
@@ -587,7 +587,7 @@ void FVoiceEngineSteam::OnAudioFinished()
 	for (FRemoteTalkerData::TIterator It(RemoteTalkerBuffers); It; ++It)
 	{
 		FRemoteTalkerDataSteam& RemoteData = It.Value();
-		if (RemoteData.VoipSynthComponent->IsIdling())
+		if (RemoteData.VoipSynthComponent && RemoteData.VoipSynthComponent->IsIdling())
 		{
 			UE_LOG_ONLINE_VOICEENGINE( Log, TEXT("Removing VOIP AudioComponent for Id: %s"), *It.Key().ToDebugString());
 			RemoteData.VoipSynthComponent->Stop();
