Reapply the fix for voice no longer broadcasting after a seamless travel on Steam.

diff --git a/Engine/Plugins/Online/OnlineSubsystemSteam/Source/Private/VoiceEngineSteam.cpp b/Engine/Plugins/Online/OnlineSubsystemSteam/Source/Private/VoiceEngineSteam.cpp
--- a/Engine/Plugins/Online/OnlineSubsystemSteam/Source/Private/VoiceEngineSteam.cpp
+++ b/Engine/Plugins/Online/OnlineSubsystemSteam/Source/Private/VoiceEngineSteam.cpp
@@ -122,6 +122,8 @@ FVoiceEngineSteam::FVoiceEngineSteam(IOnlineSubsystem* InSubsystem) :
 {
 	SteamUserPtr = SteamUser();
 	SteamFriendsPtr = SteamFriends();
+
+	FCoreUObjectDelegates::PostLoadMapWithWorld.AddRaw(this, &FVoiceEngineSteam::OnPostLoadMap);
 }
 
 FVoiceEngineSteam::~FVoiceEngineSteam()
@@ -136,6 +138,8 @@ FVoiceEngineSteam::~FVoiceEngineSteam()
 	VoiceEncoder = nullptr;
 
 	delete SerializeHelper;
+
+	FCoreUObjectDelegates::PostLoadMapWithWorld.RemoveAll(this);
 }
 
 void FVoiceEngineSteam::VoiceCaptureUpdate() const
@@ -525,7 +529,7 @@ void FVoiceEngineSteam::TickTalkers(float DeltaTime)
 			RemoteData.VoipSynthComponent->Stop();
 
 			UAudioComponent* AudioComponent = RemoteData.VoipSynthComponent->GetAudioComponent();
-			if (AudioComponent->IsRegistered())
+			if (AudioComponent != nullptr && AudioComponent->IsRegistered())
 			{
 				AudioComponent->UnregisterComponent();
 			}
@@ -598,6 +602,18 @@ void FVoiceEngineSteam::OnAudioFinished()
 	UE_LOG(LogVoiceEngine, Verbose, TEXT("Audio Finished"));
 }
 
+void FVoiceEngineSteam::OnPostLoadMap(UWorld*)
+{
+	for (FRemoteTalkerData::TIterator It(RemoteTalkerBuffers); It; ++It)
+	{
+		FRemoteTalkerDataSteam& RemoteData = It.Value();
+		if (RemoteData.VoipSynthComponent && RemoteData.VoipSynthComponent->GetAudioComponent() != nullptr)
+		{
+			RemoteData.VoipSynthComponent->GetAudioComponent()->Play();
+		}
+	}
+}
+
 FString FVoiceEngineSteam::GetVoiceDebugState() const
 {
 	FString Output;
diff --git a/Engine/Plugins/Online/OnlineSubsystemSteam/Source/Private/VoiceEngineSteam.h b/Engine/Plugins/Online/OnlineSubsystemSteam/Source/Private/VoiceEngineSteam.h
--- a/Engine/Plugins/Online/OnlineSubsystemSteam/Source/Private/VoiceEngineSteam.h
+++ b/Engine/Plugins/Online/OnlineSubsystemSteam/Source/Private/VoiceEngineSteam.h
@@ -308,14 +308,19 @@ public:
 private:
 
 	/**
-	* Update the state of all remote talkers, possibly dropping data or the talker entirely
-	*/
+	 * Update the state of all remote talkers, possibly dropping data or the talker entirely
+	 */
 	void TickTalkers(float DeltaTime);
 
 	/**
-	* Delegate triggered when an audio component Stop() function is called
-	*/
+	 * Delegate triggered when an audio component Stop() function is called
+	 */
 	void OnAudioFinished();
+
+	/**
+	 * Delegate that fixes up remote audio components when the level changes
+	 */
+	void OnPostLoadMap(UWorld*);
 };
 
 typedef TSharedPtr<FVoiceEngineSteam, ESPMode::ThreadSafe> FVoiceEngineImplPtr;
