PR #5616: Initialize FSteamAuthHandlerComponent::State to Uninitialized (Contributed by [REDACTED_PERSON])

#github https://github.com/EpicGames/UnrealEngine/pull/5616

diff --git a/Engine/Plugins/Online/OnlineSubsystemSteam/Source/Private/OnlineAuthHandlerSteam.cpp b/Engine/Plugins/Online/OnlineSubsystemSteam/Source/Private/OnlineAuthHandlerSteam.cpp
--- a/Engine/Plugins/Online/OnlineSubsystemSteam/Source/Private/OnlineAuthHandlerSteam.cpp
+++ b/Engine/Plugins/Online/OnlineSubsystemSteam/Source/Private/OnlineAuthHandlerSteam.cpp
@@ -83,6 +83,7 @@ FString ReusableKey;
 FSteamAuthHandlerComponent::FSteamAuthHandlerComponent() :
 	AuthInterface(nullptr),
 	SteamUserPtr(SteamUser()),
+	State(ESteamAuthHandlerState::Uninitialized),
 	bIsEnabled(true),
 	LastTimestamp(0.0f),
 	TicketHandle(k_HAuthTicketInvalid),
