Fix the session update issue that occurs when modifying the maximum number of connections in a session on Steam.

diff --git a/Engine/Plugins/Online/OnlineSubsystemSteam/Source/Private/OnlineSessionAsyncLobbySteam.cpp b/Engine/Plugins/Online/OnlineSubsystemSteam/Source/Private/OnlineSessionAsyncLobbySteam.cpp
--- a/Engine/Plugins/Online/OnlineSubsystemSteam/Source/Private/OnlineSessionAsyncLobbySteam.cpp
+++ b/Engine/Plugins/Online/OnlineSubsystemSteam/Source/Private/OnlineSessionAsyncLobbySteam.cpp
@@ -604,12 +604,13 @@ void FOnlineAsyncTaskSteamUpdateLobby::Tick()
 					if (SteamMatchmakingPtr->SetLobbyType(SessionInfo->SessionId, LobbyType))
 					{
 						int32 LobbyMemberCount = SteamMatchmakingPtr->GetNumLobbyMembers(SessionInfo->SessionId);
-						int32 MaxLobbyMembers = SteamMatchmakingPtr->GetLobbyMemberLimit(SessionInfo->SessionId);
-						bool bLobbyJoinable = Session->SessionSettings.bAllowJoinInProgress && (LobbyMemberCount < MaxLobbyMembers);
-						if (SteamMatchmakingPtr->SetLobbyJoinable(SessionInfo->SessionId, bLobbyJoinable))
+						int32 NumConnections = Session->SessionSettings.NumPrivateConnections + Session->SessionSettings.NumPublicConnections;
+
+						if (SteamMatchmakingPtr->SetLobbyMemberLimit(SessionInfo->SessionId, NumConnections))
 						{
-							int32 NumConnections = Session->SessionSettings.NumPrivateConnections + Session->SessionSettings.NumPublicConnections;
-							if (SteamMatchmakingPtr->SetLobbyMemberLimit(SessionInfo->SessionId, NumConnections))
+							int32 MaxLobbyMembers = SteamMatchmakingPtr->GetLobbyMemberLimit(SessionInfo->SessionId);
+							bool bLobbyJoinable = Session->SessionSettings.bAllowJoinInProgress && (LobbyMemberCount < MaxLobbyMembers) && (MaxLobbyMembers != 0);
+							if (SteamMatchmakingPtr->SetLobbyJoinable(SessionInfo->SessionId, bLobbyJoinable))
 							{
 								bWasSuccessful = true;
 
