Fix Steam Overlay invites so they work with private matches.

diff --git a/Engine/Plugins/Online/OnlineSubsystemSteam/Source/Private/OnlineExternalUIInterfaceSteam.cpp b/Engine/Plugins/Online/OnlineSubsystemSteam/Source/Private/OnlineExternalUIInterfaceSteam.cpp
--- a/Engine/Plugins/Online/OnlineSubsystemSteam/Source/Private/OnlineExternalUIInterfaceSteam.cpp
+++ b/Engine/Plugins/Online/OnlineSubsystemSteam/Source/Private/OnlineExternalUIInterfaceSteam.cpp
@@ -2,6 +2,7 @@
 
 #include "OnlineExternalUIInterfaceSteam.h"
 #include "Interfaces/OnlineSessionInterface.h"
+#include "OnlineSessionSettings.h"
 #include "OnlineSubsystemSteamTypes.h"
 
 // Other external UI possibilities in Steam
@@ -49,9 +50,26 @@ bool FOnlineExternalUISteam::ShowFriendsUI(int32 LocalUserNum)
 bool FOnlineExternalUISteam::ShowInviteUI(int32 LocalUserNum, FName SessionName)
 {
 	IOnlineSessionPtr SessionInt = SteamSubsystem->GetSessionInterface();
-	if (SessionInt.IsValid() && SessionInt->HasPresenceSession())
+	if (!SessionInt.IsValid())
 	{
-		SteamFriends()->ActivateGameOverlay("LobbyInvite");
+		return false;
+	}
+
+	const FNamedOnlineSession* const Session = SessionInt->GetNamedSession(SessionName);
+	if (Session && Session->SessionInfo.IsValid())
+	{
+		const FOnlineSessionInfoSteam* const SessionInfo = (FOnlineSessionInfoSteam*)(Session->SessionInfo.Get());
+		if (SessionInfo->SessionType == ESteamSession::LobbySession && SessionInfo->SessionId.IsValid())
+		{
+			// This can only invite to lobbies, does not work for dedicated servers.
+			SteamFriends()->ActivateGameOverlayInviteDialog(SessionInfo->SessionId);
+		}
+		else if(SessionInfo->SessionType == ESteamSession::AdvertisedSessionHost || SessionInfo->SessionType == ESteamSession::AdvertisedSessionClient)
+		{
+			// Invite people to start this game.
+			// To invite someone directly into the game, use SendSessionInviteToFriend
+			SteamFriends()->ActivateGameOverlay("LobbyInvite");
+		}
 		return true;
 	}
 
