Fix Mac failing to dynamically link with Steam API.

diff --git a/Engine/Plugins/Online/OnlineSubsystemSteam/Source/Private/OnlineSubsystemModuleSteam.cpp b/Engine/Plugins/Online/OnlineSubsystemSteam/Source/Private/OnlineSubsystemModuleSteam.cpp
--- a/Engine/Plugins/Online/OnlineSubsystemSteam/Source/Private/OnlineSubsystemModuleSteam.cpp
+++ b/Engine/Plugins/Online/OnlineSubsystemSteam/Source/Private/OnlineSubsystemModuleSteam.cpp
@@ -107,6 +107,8 @@ static FString GetSteamModulePath()
 		return FPaths::EngineDir() / STEAM_SDK_ROOT_PATH / STEAM_SDK_VER_PATH / TEXT("i686-unknown-linux-gnu/");
 	#endif	//PLATFORM_64BITS
 	
+#elif PLATFORM_MAC
+	return FPaths::EngineDir() / STEAM_SDK_ROOT_PATH / STEAM_SDK_VER_PATH / TEXT("Mac/");
 #else
 
 	return FString();
@@ -124,7 +126,7 @@ void FOnlineSubsystemSteamModule::LoadSteamModules()
 	FString Suffix("64");
 #else
 	FString Suffix;
-#endif
+#endif // PLATFORM_64BITS
 
 	FString RootSteamPath = GetSteamModulePath();
 	FPlatformProcess::PushDllDirectory(*RootSteamPath);
@@ -145,34 +147,36 @@ void FOnlineSubsystemSteamModule::LoadSteamModules()
 		}
 	}
 	FPlatformProcess::PopDllDirectory(*RootSteamPath);
-#elif PLATFORM_MAC
-	SteamDLLHandle = FPlatformProcess::GetDllHandle(TEXT("libsteam_api.dylib"));
-#elif PLATFORM_LINUX
+#elif PLATFORM_MAC || (PLATFORM_LINUX && LOADING_STEAM_LIBRARIES_DYNAMICALLY)
 
-#if LOADING_STEAM_LIBRARIES_DYNAMICALLY
-	UE_LOG_ONLINE(Log, TEXT("Loading system libsteam_api.so."));
-	SteamDLLHandle = FPlatformProcess::GetDllHandle(TEXT("libsteam_api.so"));
+#if PLATFORM_MAC
+	FString SteamModuleFileName("libsteam_api.dylib");
+#else
+	FString SteamModuleFileName("libsteam_api.so");
+#endif // PLATFORM_MAC
+
+	SteamDLLHandle = FPlatformProcess::GetDllHandle(*SteamModuleFileName);
 	if (SteamDLLHandle == nullptr)
 	{
 		// try bundled one
-		UE_LOG_ONLINE(Warning, TEXT("Could not find system one, loading bundled libsteam_api.so."));
+		UE_LOG_ONLINE(Warning, TEXT("Could not find system one, loading bundled %s."), *SteamModuleFileName);
 		FString RootSteamPath = GetSteamModulePath();
-		SteamDLLHandle = FPlatformProcess::GetDllHandle(*(RootSteamPath + "libsteam_api.so"));
+		SteamDLLHandle = FPlatformProcess::GetDllHandle(*(RootSteamPath + SteamModuleFileName));
 	}
 
 	if (SteamDLLHandle)
 	{
-		UE_LOG_ONLINE(Display, TEXT("Loaded libsteam_api.so at %p"), SteamDLLHandle);
+		UE_LOG_ONLINE(Display, TEXT("Loaded %s at %p"), *SteamModuleFileName, SteamDLLHandle);
 	}
 	else
 	{
-		UE_LOG_ONLINE(Warning, TEXT("Unable to load libsteam_api.so, Steam functionality will not work"));
+		UE_LOG_ONLINE(Warning, TEXT("Unable to load %s, Steam functionality will not work"), *SteamModuleFileName);
 	}
-#else
-	UE_LOG_ONLINE(Log, TEXT("libsteam_api.so is linked explicitly and should be already loaded."));
-#endif // LOADING_STEAM_LIBRARIES_DYNAMICALLY
 
-#endif	//PLATFORM_WINDOWS
+
+#elif PLATFORM_LINUX
+	UE_LOG_ONLINE(Log, TEXT("libsteam_api.so is linked explicitly and should be already loaded."));
+#endif // PLATFORM_WINDOWS
 }
 
 void FOnlineSubsystemSteamModule::UnloadSteamModules()
