Fixing Steam VoiceEngine so it inherits instead of duplicates code. Also fixing the FVoiceEngineImpl so it works with inheritance

diff --git a/Engine/Plugins/Online/OnlineSubsystemSteam/Source/Private/VoiceEngineSteam.cpp b/Engine/Plugins/Online/OnlineSubsystemSteam/Source/Private/VoiceEngineSteam.cpp
--- a/Engine/Plugins/Online/OnlineSubsystemSteam/Source/Private/VoiceEngineSteam.cpp
+++ b/Engine/Plugins/Online/OnlineSubsystemSteam/Source/Private/VoiceEngineSteam.cpp
@@ -1,199 +1,42 @@
 // Copyright 1998-2019 Epic Games, Inc. All Rights Reserved.
 
 #include "VoiceEngineSteam.h"
-#include "Components/AudioComponent.h"
-#include "GameFramework/GameSession.h"
 #include "OnlineSubsystemSteam.h"
 #include "OnlineSubsystemUtils.h"
-#include "Sound/SoundWaveProcedural.h"
 #include "VoiceModule.h"
-#include "Voice.h"
 #include "SteamUtilities.h"
 
-/** Largest size allowed to carry over into next buffer */
-#define MAX_VOICE_REMAINDER_SIZE 4 * 1024
-
-FRemoteTalkerDataSteam::FRemoteTalkerDataSteam() :
-	MaxUncompressedDataSize(0),
-	MaxUncompressedDataQueueSize(0),
-	CurrentUncompressedDataQueueSize(0),
-	LastSeen(0.0),
-	NumFramesStarved(0),
-	VoipSynthComponent(nullptr),
-	VoiceDecoder(nullptr)
-{
-	int32 SampleRate = UVOIPStatics::GetVoiceSampleRate();
-	int32 NumChannels = DEFAULT_NUM_VOICE_CHANNELS;
-	VoiceDecoder = FVoiceModule::Get().CreateVoiceDecoder(SampleRate, NumChannels);
-	check(VoiceDecoder.IsValid());
-
-	// Approx 1 sec worth of data for a stereo microphone
-	MaxUncompressedDataSize = UVOIPStatics::GetMaxUncompressedVoiceDataSizePerChannel() * 2;
-	MaxUncompressedDataQueueSize = MaxUncompressedDataSize * 5;
-	{
-		FScopeLock ScopeLock(&QueueLock);
-		UncompressedDataQueue.Empty(MaxUncompressedDataQueueSize);
-	}
-}
-
-FRemoteTalkerDataSteam::FRemoteTalkerDataSteam(FRemoteTalkerDataSteam&& Other)
-{
-	LastSeen = Other.LastSeen;
-	Other.LastSeen = 0.0;
-
-	NumFramesStarved = Other.NumFramesStarved;
-	Other.NumFramesStarved = 0;
-
-	VoipSynthComponent = Other.VoipSynthComponent;
-	Other.VoipSynthComponent = nullptr;
-
-	VoiceDecoder = MoveTemp(Other.VoiceDecoder);
-	Other.VoiceDecoder = nullptr;
-
-	MaxUncompressedDataSize = Other.MaxUncompressedDataSize;
-	Other.MaxUncompressedDataSize = 0;
-
-	MaxUncompressedDataQueueSize = Other.MaxUncompressedDataQueueSize;
-	Other.MaxUncompressedDataQueueSize = 0;
-
-	CurrentUncompressedDataQueueSize = Other.CurrentUncompressedDataQueueSize;
-	Other.CurrentUncompressedDataQueueSize = 0;
-
-	{
-		FScopeLock ScopeLock(&Other.QueueLock);
-		UncompressedDataQueue = MoveTemp(Other.UncompressedDataQueue);
-	}
-}
-
-FRemoteTalkerDataSteam::~FRemoteTalkerDataSteam()
-{
-	VoiceDecoder = nullptr;
-
-	CurrentUncompressedDataQueueSize = 0;
-
-	{
-		FScopeLock ScopeLock(&QueueLock);
-		UncompressedDataQueue.Empty();
-	}
-}
-
-void FRemoteTalkerDataSteam::Reset()
-{
-	// Set to large number so TickTalkers doesn't come in here
-	LastSeen = MAX_FLT;
-	NumFramesStarved = 0;
-
-	if (VoipSynthComponent)
-	{
-		VoipSynthComponent->Stop();
-		bIsActive = false;
-	}
-
-	CurrentUncompressedDataQueueSize = 0;
-
-	{
-		FScopeLock ScopeLock(&QueueLock);
-		UncompressedDataQueue.Empty();
-	}
-}
-
-void FRemoteTalkerDataSteam::Cleanup()
-{
-	if (VoipSynthComponent)
-	{
-		VoipSynthComponent->Stop();
-		bIsActive = false;
-	}
-
-	VoipSynthComponent = nullptr;
-}
-
 FVoiceEngineSteam::FVoiceEngineSteam(IOnlineSubsystem* InSubsystem) :
-	OnlineSubsystem(InSubsystem),
-	VoiceCapture(nullptr),
-	VoiceEncoder(nullptr),
-	OwningUserIndex(INVALID_INDEX),
-	UncompressedBytesAvailable(0),
-	CompressedBytesAvailable(0),
-	AvailableVoiceResult(EVoiceCaptureState::UnInitialized),
-	bPendingFinalCapture(false),
-	bIsCapturing(false),
-	SerializeHelper(nullptr)
+	FVoiceEngineImpl(InSubsystem),
+	SteamUserPtr(SteamUser()),
+	SteamFriendsPtr(SteamFriends())
 {
-	SteamUserPtr = SteamUser();
-	SteamFriendsPtr = SteamFriends();
-
-	FCoreUObjectDelegates::PostLoadMapWithWorld.AddRaw(this, &FVoiceEngineSteam::OnPostLoadMap);
 }
 
 FVoiceEngineSteam::~FVoiceEngineSteam()
 {
-	if (bIsCapturing)
+	if (IsRecording())
 	{
-		VoiceCapture->Stop();
 		SteamFriendsPtr->SetInGameVoiceSpeaking(SteamUserPtr->GetSteamID(), false);
 	}
-
-	VoiceCapture = nullptr;
-	VoiceEncoder = nullptr;
-
-	delete SerializeHelper;
-
-	FCoreUObjectDelegates::PostLoadMapWithWorld.RemoveAll(this);
-}
-
-void FVoiceEngineSteam::VoiceCaptureUpdate() const
-{
-	if (bPendingFinalCapture)
-	{
-		uint32 CompressedSize;
-		const EVoiceCaptureState::Type RecordingState = VoiceCapture->GetCaptureState(CompressedSize);
-
-		// If no data is available, we have finished capture the last (post-StopRecording) half-second of voice data
-		if (RecordingState == EVoiceCaptureState::NotCapturing)
-		{
-			UE_LOG_ONLINE_VOICEENGINE(Log, TEXT("Internal voice capture complete."));
-
-			bPendingFinalCapture = false;
-
-			// If a new recording session has begun since the call to 'StopRecording', kick that off
-			if (bIsCapturing)
-			{
-				StartRecording();
-			}
-			else
-			{
-				// Marks that recording has successfully stopped
-				StoppedRecording();
-			}
-		}
-	}
 }
 
 void FVoiceEngineSteam::StartRecording() const
 {
 	UE_LOG_ONLINE_VOICEENGINE(VeryVerbose, TEXT("VOIP StartRecording"));
-	if (VoiceCapture.IsValid())
+	if (GetVoiceCapture().IsValid())
 	{
-		if (!VoiceCapture->Start())
+		if (!GetVoiceCapture()->Start())
 		{
-			UE_LOG(LogVoiceEngine, Warning, TEXT("Failed to start voice recording"));
+			UE_LOG_ONLINE_VOICEENGINE(Warning, TEXT("Failed to start voice recording"));
 		}
-		else if (SteamFriendsPtr) {
+		else if (SteamFriendsPtr) 
+		{
 			SteamFriendsPtr->SetInGameVoiceSpeaking(SteamUserPtr->GetSteamID(), true);
 		}
 	}
 }
 
-void FVoiceEngineSteam::StopRecording() const
-{
-	UE_LOG_ONLINE_VOICEENGINE(VeryVerbose, TEXT("VOIP StopRecording"));
-	if (VoiceCapture.IsValid())
-	{
-		VoiceCapture->Stop();
-	}
-}
-
 void FVoiceEngineSteam::StoppedRecording() const
 {
 	UE_LOG_ONLINE_VOICEENGINE(VeryVerbose, TEXT("VOIP StoppedRecording"));
@@ -203,512 +46,3 @@ void FVoiceEngineSteam::StoppedRecording() const
 	}
 }
 
-bool FVoiceEngineSteam::Init(int32 MaxLocalTalkers, int32 MaxRemoteTalkers)
-{
-	bool bSuccess = false;
-
-	if (!OnlineSubsystem->IsDedicated())
-	{
-		FVoiceModule& VoiceModule = FVoiceModule::Get();
-		if (VoiceModule.IsVoiceEnabled())
-		{
-			VoiceCapture = VoiceModule.CreateVoiceCapture();
-			VoiceEncoder = VoiceModule.CreateVoiceEncoder();
-
-			bSuccess = VoiceCapture.IsValid() && VoiceEncoder.IsValid();
-			if (bSuccess)
-			{
-				CompressedVoiceBuffer.Empty(UVOIPStatics::GetMaxCompressedVoiceDataSize());
-				DecompressedVoiceBuffer.Empty(UVOIPStatics::GetMaxUncompressedVoiceDataSizePerChannel());
-
-				for (int32 TalkerIdx = 0; TalkerIdx < MaxLocalTalkers; TalkerIdx++)
-				{
-					PlayerVoiceData[TalkerIdx].VoiceRemainderSize = 0;
-					PlayerVoiceData[TalkerIdx].VoiceRemainder.Empty(MAX_VOICE_REMAINDER_SIZE);
-				}
-			}
-			else
-			{
-				UE_LOG(LogVoice, Warning, TEXT("Voice capture initialization failed!"));
-			}
-		}
-		else
-		{
-			UE_LOG(LogVoice, Log, TEXT("Voice module disabled by config [Voice].bEnabled"));
-		}
-	}
-
-	return bSuccess;
-}
-
-uint32 FVoiceEngineSteam::StartLocalVoiceProcessing(uint32 LocalUserNum)
-{
-	uint32 Return = ONLINE_FAIL;
-	if (IsOwningUser(LocalUserNum))
-	{
-		if (!bIsCapturing)
-		{
-			// Update the current recording state, if VOIP data was still being read
-			VoiceCaptureUpdate();
-
-			if (!IsRecording())
-			{
-				StartRecording();
-			}
-
-			bIsCapturing = true;
-		}
-
-		Return = ONLINE_SUCCESS;
-	}
-	else
-	{
-		UE_LOG_ONLINE_VOICEENGINE(Error, TEXT("StartLocalVoiceProcessing(): Device is currently owned by another user"));
-	}
-
-	return Return;
-}
-
-uint32 FVoiceEngineSteam::StopLocalVoiceProcessing(uint32 LocalUserNum)
-{
-	uint32 Return = ONLINE_FAIL;
-	if (IsOwningUser(LocalUserNum))
-	{
-		if (bIsCapturing)
-		{
-			bIsCapturing = false;
-			bPendingFinalCapture = true;
-
-			// Make a call to begin stopping the current VOIP recording session
-			StopRecording();
-
-			// Now check/update the status of the recording session
-			VoiceCaptureUpdate();
-		}
-
-		Return = ONLINE_SUCCESS;
-	}
-	else
-	{
-		UE_LOG_ONLINE_VOICEENGINE(Error, TEXT("StopLocalVoiceProcessing: Ignoring stop request for non-owning user"));
-	}
-
-	return Return;
-}
-
-uint32 FVoiceEngineSteam::UnregisterRemoteTalker(const FUniqueNetId& UniqueId)
-{
-	FRemoteTalkerDataSteam* RemoteData = RemoteTalkerBuffers.Find(FUniqueNetIdWrapper(UniqueId.AsShared()));
-	if (RemoteData != nullptr)
-	{
-		// Dump the whole talker
-		RemoteData->Cleanup();
-		RemoteTalkerBuffers.Remove(FUniqueNetIdWrapper(UniqueId.AsShared()));
-	}
-
-	return ONLINE_SUCCESS;
-}
-
-uint32 FVoiceEngineSteam::GetVoiceDataReadyFlags() const
-{
-	// First check and update the internal state of VOIP recording
-	VoiceCaptureUpdate();
-	if (OwningUserIndex != INVALID_INDEX && IsRecording())
-	{
-		// Check if there is new data available via the Voice API
-		if (AvailableVoiceResult == EVoiceCaptureState::Ok && UncompressedBytesAvailable > 0)
-		{
-			return 1 << OwningUserIndex;
-		}
-	}
-
-	return 0;
-}
-
-uint32 FVoiceEngineSteam::ReadLocalVoiceData(uint32 LocalUserNum, uint8* Data, uint32* Size, uint64* OutSampleCount)
-{
-	check(*Size > 0);
-
-	// Before doing anything, check/update the current recording state
-	VoiceCaptureUpdate();
-
-	// Return data even if not capturing, possibly have data during stopping
-	if (IsOwningUser(LocalUserNum) && IsRecording())
-	{
-		DecompressedVoiceBuffer.Empty(UVOIPStatics::GetMaxUncompressedVoiceDataSizePerChannel());
-		CompressedVoiceBuffer.Empty(UVOIPStatics::GetMaxCompressedVoiceDataSize());
-
-		uint32 NewVoiceDataBytes = 0;
-		EVoiceCaptureState::Type VoiceResult = VoiceCapture->GetCaptureState(NewVoiceDataBytes);
-		if (VoiceResult != EVoiceCaptureState::Ok && VoiceResult != EVoiceCaptureState::NoData)
-		{
-			UE_LOG_ONLINE_VOICEENGINE(Warning, TEXT("ReadLocalVoiceData: GetAvailableVoice failure: VoiceResult: %s"), EVoiceCaptureState::ToString(VoiceResult));
-			return ONLINE_FAIL;
-		}
-
-		if (NewVoiceDataBytes == 0)
-		{
-			UE_LOG_ONLINE_VOICEENGINE(VeryVerbose, TEXT("ReadLocalVoiceData: No Data: VoiceResult: %s"), EVoiceCaptureState::ToString(VoiceResult));
-			*Size = 0;
-			return ONLINE_SUCCESS;
-		}
-
-		// Make space for new and any previously remaining data
-
-		// Add the number of new bytes (since last time this function was called) and the number of bytes remaining that wasn't consumed last time this was called
-		// This is how many bytes we would like to return
-		uint32 TotalVoiceBytes = NewVoiceDataBytes + PlayerVoiceData[LocalUserNum].VoiceRemainderSize;
-
-		// But we have a max amount we can return so clamp it to that max value if we're asking for more bytes than we're allowed
-		if (TotalVoiceBytes > UVOIPStatics::GetMaxUncompressedVoiceDataSizePerChannel())
-		{
-			UE_LOG(LogVoiceEngine, Warning, TEXT("Exceeded uncompressed voice buffer size, clamping"))
-				TotalVoiceBytes = UVOIPStatics::GetMaxUncompressedVoiceDataSizePerChannel();
-		}
-
-		DecompressedVoiceBuffer.AddUninitialized(TotalVoiceBytes);
-
-		// If there's still audio left from a previous ReadLocalData call that didn't get output, copy that first into the decompressed voice buffer
-		if (PlayerVoiceData[LocalUserNum].VoiceRemainderSize > 0)
-		{
-			FMemory::Memcpy(DecompressedVoiceBuffer.GetData(), PlayerVoiceData[LocalUserNum].VoiceRemainder.GetData(), PlayerVoiceData[LocalUserNum].VoiceRemainderSize);
-		}
-
-		// Get new uncompressed data
-		uint8* RemainingDecompressedBufferPtr = DecompressedVoiceBuffer.GetData() + PlayerVoiceData[LocalUserNum].VoiceRemainderSize;
-		uint32 ByteWritten = 0;
-		uint64 NewSampleCount = 0;
-		VoiceResult = VoiceCapture->GetVoiceData(DecompressedVoiceBuffer.GetData() + PlayerVoiceData[LocalUserNum].VoiceRemainderSize, NewVoiceDataBytes, ByteWritten, NewSampleCount);
-
-		TotalVoiceBytes = ByteWritten + PlayerVoiceData[LocalUserNum].VoiceRemainderSize;
-
-		if ((VoiceResult == EVoiceCaptureState::Ok || VoiceResult == EVoiceCaptureState::NoData) && TotalVoiceBytes > 0)
-		{
-			if (OutSampleCount != nullptr)
-			{
-				*OutSampleCount = NewSampleCount;
-			}
-
-			// Prepare the encoded buffer (e.g. opus)
-			CompressedBytesAvailable = UVOIPStatics::GetMaxCompressedVoiceDataSize();
-			CompressedVoiceBuffer.AddUninitialized(UVOIPStatics::GetMaxCompressedVoiceDataSize());
-
-			check(((uint32)CompressedVoiceBuffer.Num()) <= UVOIPStatics::GetMaxCompressedVoiceDataSize());
-
-			// Run the uncompressed audio through the opus decoder, note that it may not encode all data, which results in some remaining data
-			PlayerVoiceData[LocalUserNum].VoiceRemainderSize =
-				VoiceEncoder->Encode(DecompressedVoiceBuffer.GetData(), TotalVoiceBytes, CompressedVoiceBuffer.GetData(), CompressedBytesAvailable);
-
-			// Save off any unencoded remainder
-			if (PlayerVoiceData[LocalUserNum].VoiceRemainderSize > 0)
-			{
-				if (PlayerVoiceData[LocalUserNum].VoiceRemainderSize > MAX_VOICE_REMAINDER_SIZE)
-				{
-					UE_LOG(LogVoiceEngine, Warning, TEXT("Exceeded voice remainder buffer size, clamping"));
-					PlayerVoiceData[LocalUserNum].VoiceRemainderSize = MAX_VOICE_REMAINDER_SIZE;
-				}
-
-				PlayerVoiceData[LocalUserNum].VoiceRemainder.AddUninitialized(MAX_VOICE_REMAINDER_SIZE);
-				FMemory::Memcpy(PlayerVoiceData[LocalUserNum].VoiceRemainder.GetData(), DecompressedVoiceBuffer.GetData() + (TotalVoiceBytes - PlayerVoiceData[LocalUserNum].VoiceRemainderSize), PlayerVoiceData[LocalUserNum].VoiceRemainderSize);
-			}
-
-			static double LastGetVoiceCallTime = 0.0;
-			double CurTime = FPlatformTime::Seconds();
-			double TimeSinceLastCall = (LastGetVoiceCallTime > 0) ? (CurTime - LastGetVoiceCallTime) : 0.0;
-			LastGetVoiceCallTime = CurTime;
-
-			UE_LOG(LogVoiceEngine, Log, TEXT("ReadLocalVoiceData: GetVoice: Result: %s, Available: %i, LastCall: %0.3f ms"), EVoiceCaptureState::ToString(VoiceResult), CompressedBytesAvailable, TimeSinceLastCall * 1000.0);
-			if (CompressedBytesAvailable > 0)
-			{
-				*Size = FMath::Min<int32>(*Size, CompressedBytesAvailable);
-				FMemory::Memcpy(Data, CompressedVoiceBuffer.GetData(), *Size);
-
-				UE_LOG(LogVoiceEngine, VeryVerbose, TEXT("ReadLocalVoiceData: Size: %d"), *Size);
-				return ONLINE_SUCCESS;
-			}
-			else
-			{
-				*Size = 0;
-				CompressedVoiceBuffer.Empty(UVOIPStatics::GetMaxCompressedVoiceDataSize());
-
-				UE_LOG_ONLINE_VOICEENGINE(Warning, TEXT("ReadLocalVoiceData: GetVoice failure: VoiceResult: %s"), EVoiceCaptureState::ToString(VoiceResult));
-				return ONLINE_FAIL;
-			}
-		}
-	}
-
-	return ONLINE_FAIL;
-}
-
-uint32 FVoiceEngineSteam::SubmitRemoteVoiceData(const FUniqueNetIdWrapper& RemoteTalkerId, uint8* Data, uint32* Size, uint64& InSampleCount)
-{
-	UE_LOG_ONLINE_VOICEENGINE(VeryVerbose, TEXT("SubmitRemoteVoiceData(%s) Size: %d received!"), *RemoteTalkerId.ToDebugString(), *Size);
-
-	FRemoteTalkerDataSteam& QueuedData = RemoteTalkerBuffers.FindOrAdd(RemoteTalkerId);
-
-	// new voice packet.
-	QueuedData.LastSeen = FPlatformTime::Seconds();
-
-	uint32 BytesWritten = UVOIPStatics::GetMaxUncompressedVoiceDataSizePerChannel();
-
-	DecompressedVoiceBuffer.Empty(UVOIPStatics::GetMaxUncompressedVoiceDataSizePerChannel());
-	DecompressedVoiceBuffer.AddUninitialized(UVOIPStatics::GetMaxUncompressedVoiceDataSizePerChannel());
-	QueuedData.VoiceDecoder->Decode(Data, *Size, DecompressedVoiceBuffer.GetData(), BytesWritten);
-
-	// If there is no data, return
-	if (BytesWritten <= 0)
-	{
-		*Size = 0;
-		return ONLINE_SUCCESS;
-	}
-
-	bool bAudioComponentCreated = false;
-	// Generate a streaming wave audio component for voice playback
-	if (QueuedData.VoipSynthComponent == nullptr || QueuedData.VoipSynthComponent->IsPendingKill())
-	{
-		if (SerializeHelper == nullptr)
-		{
-			SerializeHelper = new FVoiceSerializeHelper(this);
-		}
-
-		QueuedData.VoipSynthComponent = CreateVoiceSynthComponent(UVOIPStatics::GetVoiceSampleRate());
-		if (QueuedData.VoipSynthComponent)
-		{
-			//TODO, make buffer size and buffering delay runtime-controllable parameters.
-			QueuedData.bIsActive = false;
-			QueuedData.VoipSynthComponent->OpenPacketStream(InSampleCount, UVOIPStatics::GetNumBufferedPackets(), UVOIPStatics::GetBufferingDelay());
-			QueuedData.bIsEnvelopeBound = false;
-		}
-	}
-
-	if (QueuedData.VoipSynthComponent != nullptr)
-	{
-		if (!QueuedData.bIsActive)
-		{
-			QueuedData.bIsActive = true;
-			FVoiceSettings InSettings;
-			UVOIPTalker* OwningTalker = nullptr;
-
-			OwningTalker = UVOIPStatics::GetVOIPTalkerForPlayer(RemoteTalkerId, InSettings);
-
-			ApplyVoiceSettings(QueuedData.VoipSynthComponent, InSettings);
-
-			QueuedData.VoipSynthComponent->ResetBuffer(InSampleCount, UVOIPStatics::GetBufferingDelay());
-			QueuedData.VoipSynthComponent->Start();
-			QueuedData.CachedTalkerPtr = OwningTalker;
-
-			if (OwningTalker)
-			{
-				if (!QueuedData.bIsEnvelopeBound)
-				{
-					QueuedData.VoipSynthComponent->OnAudioEnvelopeValueNative.AddUObject(OwningTalker, &UVOIPTalker::OnAudioComponentEnvelopeValue);
-					QueuedData.bIsEnvelopeBound = true;
-				}
-
-				OwningTalker->OnTalkingBegin(QueuedData.VoipSynthComponent->GetAudioComponent());
-			}
-		}
-
-		QueuedData.VoipSynthComponent->SubmitPacket((float*)DecompressedVoiceBuffer.GetData(), BytesWritten, InSampleCount, EVoipStreamDataFormat::Int16);
-	}
-
-	return ONLINE_SUCCESS;
-}
-
-void FVoiceEngineSteam::TickTalkers(float DeltaTime)
-{
-	// Remove users that are done talking.
-	const double CurTime = FPlatformTime::Seconds();
-	for (FRemoteTalkerData::TIterator It(RemoteTalkerBuffers); It; ++It)
-	{
-		FRemoteTalkerDataSteam& RemoteData = It.Value();
-		double TimeSince = CurTime - RemoteData.LastSeen;
-
-		if (RemoteData.VoipSynthComponent && RemoteData.VoipSynthComponent->IsIdling() && RemoteData.bIsActive)
-		{
-			RemoteData.VoipSynthComponent->Stop();
-
-			UAudioComponent* AudioComponent = RemoteData.VoipSynthComponent->GetAudioComponent();
-			if (AudioComponent != nullptr && AudioComponent->IsRegistered())
-			{
-				AudioComponent->UnregisterComponent();
-			}
-
-			//If the UVOIPTalker associated with this is still alive, notify it that this player is done talking.
-			if (UVOIPStatics::IsVOIPTalkerStillAlive(RemoteData.CachedTalkerPtr))
-			{
-				RemoteData.CachedTalkerPtr->OnTalkingEnd();
-			}
-			RemoteData.bIsActive = false;
-
-			RemoteData.Reset();
-		}
-		else if (TimeSince >= UVOIPStatics::GetRemoteTalkerTimeoutDuration())
-		{
-			// Dump the whole talker
-			RemoteData.Reset();
-		}
-	}
-}
-
-void FVoiceEngineSteam::Tick(float DeltaTime)
-{
-	// Check available voice once a frame, this value changes after calling GetVoiceData()
-	AvailableVoiceResult = VoiceCapture->GetCaptureState(UncompressedBytesAvailable);
-
-	TickTalkers(DeltaTime);
-}
-
-void FVoiceEngineSteam::GenerateVoiceData(USoundWaveProcedural* InProceduralWave, int32 SamplesRequired, const FUniqueNetId& TalkerId)
-{
-	FRemoteTalkerDataSteam* QueuedData = RemoteTalkerBuffers.Find(FUniqueNetIdWrapper(TalkerId.AsShared()));
-	if (QueuedData)
-	{
-		const int32 SampleSize = sizeof(uint16) * DEFAULT_NUM_VOICE_CHANNELS;
-
-		{
-			FScopeLock ScopeLock(&QueuedData->QueueLock);
-			QueuedData->CurrentUncompressedDataQueueSize = QueuedData->UncompressedDataQueue.Num();
-			const int32 AvailableSamples = QueuedData->CurrentUncompressedDataQueueSize / SampleSize;
-			if (AvailableSamples >= SamplesRequired)
-			{
-				UE_LOG_ONLINE_VOICEENGINE(Verbose, TEXT("GenerateVoiceData %d / %d"), AvailableSamples, SamplesRequired);
-				const int32 SamplesBytesTaken = AvailableSamples * SampleSize;
-				InProceduralWave->QueueAudio(QueuedData->UncompressedDataQueue.GetData(), SamplesBytesTaken);
-				QueuedData->UncompressedDataQueue.RemoveAt(0, SamplesBytesTaken, false);
-				QueuedData->CurrentUncompressedDataQueueSize -= (SamplesBytesTaken);
-			}
-			else
-			{
-				UE_LOG_ONLINE_VOICEENGINE(Verbose, TEXT("Voice underflow"));
-			}
-		}
-	}
-}
-
-void FVoiceEngineSteam::OnAudioFinished()
-{
-	for (FRemoteTalkerData::TIterator It(RemoteTalkerBuffers); It; ++It)
-	{
-		FRemoteTalkerDataSteam& RemoteData = It.Value();
-		if (RemoteData.VoipSynthComponent && RemoteData.VoipSynthComponent->IsIdling())
-		{
-			UE_LOG_ONLINE_VOICEENGINE( Log, TEXT("Removing VOIP AudioComponent for Id: %s"), *It.Key().ToDebugString());
-			RemoteData.VoipSynthComponent->Stop();
-			RemoteData.bIsActive = false;
-			break;
-		}
-	}
-	UE_LOG(LogVoiceEngine, Verbose, TEXT("Audio Finished"));
-}
-
-void FVoiceEngineSteam::OnPostLoadMap(UWorld*)
-{
-	for (FRemoteTalkerData::TIterator It(RemoteTalkerBuffers); It; ++It)
-	{
-		FRemoteTalkerDataSteam& RemoteData = It.Value();
-		if (RemoteData.VoipSynthComponent && RemoteData.VoipSynthComponent->GetAudioComponent() != nullptr)
-		{
-			RemoteData.VoipSynthComponent->GetAudioComponent()->Play();
-		}
-	}
-}
-
-FString FVoiceEngineSteam::GetVoiceDebugState() const
-{
-	FString Output;
-	Output = FString::Printf(TEXT("IsRecording: %d\n DataReady: 0x%08x State:%s\n UncompressedBytes: %d\n CompressedBytes: %d\n"),
-		IsRecording(),
-		GetVoiceDataReadyFlags(),
-		EVoiceCaptureState::ToString(AvailableVoiceResult),
-		UncompressedBytesAvailable,
-		CompressedBytesAvailable
-	);
-
-	// Add remainder size
-	for (int32 Idx = 0; Idx < MAX_SPLITSCREEN_TALKERS; Idx++)
-	{
-		Output += FString::Printf(TEXT("Remainder[%d] %d\n"), Idx, PlayerVoiceData[Idx].VoiceRemainderSize);
-	}
-
-	return Output;
-}
-
-bool FVoiceEngineSteam::Exec(UWorld* InWorld, const TCHAR* Cmd, FOutputDevice& Ar)
-{
-	bool bWasHandled = false;
-
-	if (FParse::Command(&Cmd, TEXT("vcvbr")))
-	{
-		// vcvbr <true/false>
-		FString VBRStr = FParse::Token(Cmd, false);
-		int32 ShouldVBR = FPlatformString::Atoi(*VBRStr);
-		bool bVBR = ShouldVBR != 0;
-		if (VoiceEncoder.IsValid())
-		{
-			if (!VoiceEncoder->SetVBR(bVBR))
-			{
-				UE_LOG(LogVoice, Warning, TEXT("Failed to set VBR %d"), bVBR);
-			}
-		}
-
-		bWasHandled = true;
-	}
-	else if (FParse::Command(&Cmd, TEXT("vcbitrate")))
-	{
-		// vcbitrate <bitrate>
-		FString BitrateStr = FParse::Token(Cmd, false);
-		int32 NewBitrate = !BitrateStr.IsEmpty() ? FPlatformString::Atoi(*BitrateStr) : 0;
-		if (VoiceEncoder.IsValid() && NewBitrate > 0)
-		{
-			if (!VoiceEncoder->SetBitrate(NewBitrate))
-			{
-				UE_LOG(LogVoice, Warning, TEXT("Failed to set bitrate %d"), NewBitrate);
-			}
-		}
-
-		bWasHandled = true;
-	}
-	else if (FParse::Command(&Cmd, TEXT("vccomplexity")))
-	{
-		// vccomplexity <complexity>
-		FString ComplexityStr = FParse::Token(Cmd, false);
-		int32 NewComplexity = !ComplexityStr.IsEmpty() ? FPlatformString::Atoi(*ComplexityStr) : -1;
-		if (VoiceEncoder.IsValid() && NewComplexity >= 0)
-		{
-			if (!VoiceEncoder->SetComplexity(NewComplexity))
-			{
-				UE_LOG(LogVoice, Warning, TEXT("Failed to set complexity %d"), NewComplexity);
-			}
-		}
-
-		bWasHandled = true;
-	}
-	else if (FParse::Command(&Cmd, TEXT("vcdump")))
-	{
-		if (VoiceCapture.IsValid())
-		{
-			VoiceCapture->DumpState();
-		}
-
-		if (VoiceEncoder.IsValid())
-		{
-			VoiceEncoder->DumpState();
-		}
-
-		for (FRemoteTalkerData::TIterator It(RemoteTalkerBuffers); It; ++It)
-		{
-			FRemoteTalkerDataSteam& RemoteData = It.Value();
-			if (RemoteData.VoiceDecoder.IsValid())
-			{
-				RemoteData.VoiceDecoder->DumpState();
-			}
-		}
-
-		bWasHandled = true;
-	}
-
-	return bWasHandled;
-}
-
-
diff --git a/Engine/Plugins/Online/OnlineSubsystemSteam/Source/Private/VoiceEngineSteam.h b/Engine/Plugins/Online/OnlineSubsystemSteam/Source/Private/VoiceEngineSteam.h
--- a/Engine/Plugins/Online/OnlineSubsystemSteam/Source/Private/VoiceEngineSteam.h
+++ b/Engine/Plugins/Online/OnlineSubsystemSteam/Source/Private/VoiceEngineSteam.h
@@ -3,324 +3,43 @@
 #pragma once
 
 #include "CoreMinimal.h"
-#include "UObject/GCObject.h"
 #include "OnlineSubsystemTypes.h"
 #include "OnlineSubsystemSteamTypes.h"
-#include "Interfaces/VoiceInterface.h"
-#include "Net/VoiceDataCommon.h"
-#include "Interfaces/VoiceCapture.h"
-#include "Interfaces/VoiceCodec.h"
 #include "OnlineSubsystemUtilsPackage.h"
-#include "VoipListenerSynthComponent.h"
+#include "VoiceEngineImpl.h"
 
 class IOnlineSubsystem;
 class FUniqueNetIdSteam;
-class IVoiceDecoder;
-class IVoiceEncoder;
-class IVoiceCapture;
 
 #define INVALID_INDEX -1
 
-/**
-* Container for unprocessed voice data
-*/
-struct FLocalVoiceDataSteam
-{
-	FLocalVoiceDataSteam() :
-		VoiceRemainderSize(0)
-	{
-	}
-
-	/** Amount of voice data not encoded last time */
-	uint32 VoiceRemainderSize;
-	/** Voice sample data not encoded last time */
-	TArray<uint8> VoiceRemainder;
-};
-
-/**
-* Remote voice data playing on a single client
-*/
-class FRemoteTalkerDataSteam
-{
-public:
-
-	FRemoteTalkerDataSteam();
-	/** Required for TMap FindOrAdd() */
-	FRemoteTalkerDataSteam(FRemoteTalkerDataSteam&& Other);
-	~FRemoteTalkerDataSteam();
-
-	/** Reset the talker after long periods of silence */
-	void Reset();
-	/** Cleanup the talker before unregistration */
-	void Cleanup();
-
-	/** Maximum size of a single decoded packet */
-	int32 MaxUncompressedDataSize;
-	/** Maximum size of the outgoing playback queue */
-	int32 MaxUncompressedDataQueueSize;
-	/** Amount of data currently in the outgoing playback queue */
-	int32 CurrentUncompressedDataQueueSize;
-
-	/** Receive side timestamp since last voice packet fragment */
-	double LastSeen;
-	/** Number of frames starved of audio */
-	int32 NumFramesStarved;
-	/** Synth component playing this buffer (only valid on remote instances) */
-	UVoipListenerSynthComponent* VoipSynthComponent;
-	/** Cached Talker Ptr. Is checked against map before use to ensure it has not been destroyed. */
-	UVOIPTalker* CachedTalkerPtr;
-	/** Boolean used to ensure that we only bind the VOIP talker to the SynthComponent's corresponding envelope delegate once. */
-	bool bIsEnvelopeBound;
-	/** Boolean flag used to tell whether this synth component is currently consuming incoming voice packets. */
-	bool bIsActive;
-	/** Buffer for outgoing audio intended for procedural streaming */
-	mutable FCriticalSection QueueLock;
-	TArray<uint8> UncompressedDataQueue;
-	/** Per remote talker voice decoding state */
-	TSharedPtr<IVoiceDecoder> VoiceDecoder;
-};
-
 /**
 * Generic implementation of voice engine, using Voice module for capture/codec
 */
-class FVoiceEngineSteam : public IVoiceEngine, public FSelfRegisteringExec
+class FVoiceEngineSteam : public FVoiceEngineImpl
 {
-	class FVoiceSerializeHelper : public FGCObject
-	{
-		/** Reference to audio components */
-		FVoiceEngineSteam* VoiceEngine;
-		FVoiceSerializeHelper() :
-			VoiceEngine(nullptr)
-		{}
-
-	public:
-
-		FVoiceSerializeHelper(FVoiceEngineSteam* InVoiceEngine) :
-			VoiceEngine(InVoiceEngine)
-		{}
-		~FVoiceSerializeHelper() {}
-
-		/** FGCObject interface */
-		virtual void AddReferencedObjects(FReferenceCollector& Collector) override
-		{
-			// Prevent garbage collection of audio components
-			for (FRemoteTalkerData::TIterator It(VoiceEngine->RemoteTalkerBuffers); It; ++It)
-			{
-				FRemoteTalkerDataSteam& RemoteData = It.Value();
-				if (RemoteData.VoipSynthComponent)
-				{
-					Collector.AddReferencedObject(RemoteData.VoipSynthComponent);
-				}
-			}
-		}
-	};
-
-	friend class FVoiceSerializeHelper;
-
-	/** Mapping of UniqueIds to the incoming voice data and their audio component */
-	typedef TMap<FUniqueNetIdWrapper, FRemoteTalkerDataSteam> FRemoteTalkerData;
-
-	/** Reference to the main online subsystem */
-	IOnlineSubsystem* OnlineSubsystem;
-
 	/** Steam User interface */
 	class ISteamUser* SteamUserPtr;
 	/** Steam Friends interface */
 	class ISteamFriends* SteamFriendsPtr;
 
-	FLocalVoiceDataSteam PlayerVoiceData[MAX_SPLITSCREEN_TALKERS];
-	/** Reference to voice capture device */
-	TSharedPtr<IVoiceCapture> VoiceCapture;
-	/** Reference to voice encoding object */
-	TSharedPtr<IVoiceEncoder> VoiceEncoder;
-
-	/** User index currently holding onto the voice interface */
-	int32 OwningUserIndex;
-	/** Amount of uncompressed data available this frame */
-	uint32 UncompressedBytesAvailable;
-	/** Amount of compressed data available this frame */
-	uint32 CompressedBytesAvailable;
-	/** Current frame state of voice capture */
-	EVoiceCaptureState::Type AvailableVoiceResult;
-	/** Have we stopped capturing voice but are waiting for its completion */
-	mutable bool bPendingFinalCapture;
-	/** State of voice recording */
-	bool bIsCapturing;
-
-	/** Data from voice codec, waiting to send to network. */
-	TArray<uint8> CompressedVoiceBuffer;
-	/** Data from network playing on an audio component. */
-	FRemoteTalkerData RemoteTalkerBuffers;
-	/** Voice decompression buffer, shared by all talkers, valid during SubmitRemoteVoiceData */
-	TArray<uint8> DecompressedVoiceBuffer;
-	/** Serialization helper */
-	FVoiceSerializeHelper* SerializeHelper;
-
-	/**
-	* Determines if the specified index is the owner or not
-	*
-	* @param InIndex the index being tested
-	*
-	* @return true if this is the owner, false otherwise
-	*/
-	FORCEINLINE bool IsOwningUser(uint32 UserIndex)
-	{
-		return UserIndex >= 0 && UserIndex < MAX_SPLITSCREEN_TALKERS && OwningUserIndex == UserIndex;
-	}
-
-	/**
-	* Update the internal state of the voice capturing state
-	* Handles possible continuation waiting for capture stop event
-	*/
-	void VoiceCaptureUpdate() const;
-
 	/** Start capturing voice data */
-	void StartRecording() const;
-
-	/** Stop capturing voice data */
-	void StopRecording() const;
+	virtual void StartRecording() const override;
 
 	/** Called when "last half second" is over */
-	void StoppedRecording() const;
-
-	/** @return is active recording occurring at the moment */
-	bool IsRecording() const { return bIsCapturing || bPendingFinalCapture; }
-
-	/**
-	* Callback from streaming audio when data is requested for playback
-	*
-	* @param InProceduralWave SoundWave requesting more data
-	* @param SamplesRequired number of samples needed for immediate playback
-	* @param TalkerId id of the remote talker to allocate voice data for
-	*/
-	void GenerateVoiceData(USoundWaveProcedural* InProceduralWave, int32 SamplesRequired, const FUniqueNetId& TalkerId);
+	virtual void StoppedRecording() const override;
 
 PACKAGE_SCOPE:
 
 	/** Constructor */
 	FVoiceEngineSteam() :
-		OnlineSubsystem(nullptr),
+		FVoiceEngineImpl(),
 		SteamUserPtr(nullptr),
-		SteamFriendsPtr(nullptr),
-		VoiceCapture(nullptr),
-		VoiceEncoder(nullptr),
-		OwningUserIndex(INVALID_INDEX),
-		UncompressedBytesAvailable(0),
-		CompressedBytesAvailable(0),
-		AvailableVoiceResult(EVoiceCaptureState::UnInitialized),
-		bPendingFinalCapture(false),
-		bIsCapturing(false),
-		SerializeHelper(nullptr)
+		SteamFriendsPtr(nullptr)
 	{};
 
-	// IVoiceEngine
-	virtual bool Init(int32 MaxLocalTalkers, int32 MaxRemoteTalkers) override;
-
 public:
 
 	FVoiceEngineSteam(IOnlineSubsystem* InSubsystem);
 	virtual ~FVoiceEngineSteam();
-
-	// IVoiceEngine
-	virtual uint32 StartLocalVoiceProcessing(uint32 LocalUserNum) override;
-	virtual uint32 StopLocalVoiceProcessing(uint32 LocalUserNum) override;
-
-	virtual uint32 StartRemoteVoiceProcessing(const FUniqueNetId& UniqueId) override
-	{
-		// Not needed
-		return ONLINE_SUCCESS;
-	}
-
-	virtual uint32 StopRemoteVoiceProcessing(const FUniqueNetId& UniqueId) override
-	{
-		// Not needed
-		return ONLINE_SUCCESS;
-	}
-
-	virtual uint32 RegisterLocalTalker(uint32 LocalUserNum) override
-	{
-		if (OwningUserIndex == INVALID_INDEX)
-		{
-			OwningUserIndex = LocalUserNum;
-			return ONLINE_SUCCESS;
-		}
-
-		return ONLINE_FAIL;
-	}
-
-	virtual uint32 UnregisterLocalTalker(uint32 LocalUserNum) override
-	{
-		if (IsOwningUser(LocalUserNum))
-		{
-			OwningUserIndex = INVALID_INDEX;
-			return ONLINE_SUCCESS;
-		}
-
-		return ONLINE_FAIL;
-	}
-
-	virtual uint32 RegisterRemoteTalker(const FUniqueNetId& UniqueId) override
-	{
-		// Not needed
-		return ONLINE_SUCCESS;
-	}
-
-	virtual uint32 UnregisterRemoteTalker(const FUniqueNetId& UniqueId) override;
-
-	virtual bool IsHeadsetPresent(uint32 LocalUserNum) override
-	{
-		return IsOwningUser(LocalUserNum) ? true : false;
-	}
-
-	virtual bool IsLocalPlayerTalking(uint32 LocalUserNum) override
-	{
-		return (GetVoiceDataReadyFlags() & (LocalUserNum << 1)) != 0;
-	}
-
-	virtual bool IsRemotePlayerTalking(const FUniqueNetId& UniqueId) override
-	{
-		return RemoteTalkerBuffers.Find(FUniqueNetIdWrapper(UniqueId.AsShared())) != nullptr;
-	}
-
-	virtual uint32 GetVoiceDataReadyFlags() const override;
-	virtual uint32 SetPlaybackPriority(uint32 LocalUserNum, const FUniqueNetId& RemoteTalkerId, uint32 Priority) override
-	{
-		// Not supported
-		return ONLINE_SUCCESS;
-	}
-
-	virtual uint32 ReadLocalVoiceData(uint32 LocalUserNum, uint8* Data, uint32* Size) override { return ReadLocalVoiceData(LocalUserNum, Data, Size, nullptr); }
-	virtual uint32 ReadLocalVoiceData(uint32 LocalUserNum, uint8* Data, uint32* Size, uint64* OutSampleCount) override;
-
-	virtual uint32 SubmitRemoteVoiceData(const FUniqueNetId& RemoteTalkerId, uint8* Data, uint32* Size)
-	{
-		checkf(false, TEXT("Please use the following function signature instead: SubmitRemoteVoiceData(const FUniqueNetIdWrapper& RemoteTalkerId, uint8* Data, uint32* Size, uint64& InSampleCount)"));
-		return 0;
-	}
-	virtual uint32 SubmitRemoteVoiceData(const FUniqueNetIdWrapper& RemoteTalkerId, uint8* Data, uint32* Size, uint64& InSampleCount) override;
-
-	virtual void Tick(float DeltaTime) override;
-	FString GetVoiceDebugState() const override;
-
-	// FSelfRegisteringExec
-	virtual bool Exec(UWorld* InWorld, const TCHAR* Cmd, FOutputDevice& Ar) override;
-
-private:
-
-	/**
-	 * Update the state of all remote talkers, possibly dropping data or the talker entirely
-	 */
-	void TickTalkers(float DeltaTime);
-
-	/**
-	 * Delegate triggered when an audio component Stop() function is called
-	 */
-	void OnAudioFinished();
-
-	/**
-	 * Delegate that fixes up remote audio components when the level changes
-	 */
-	void OnPostLoadMap(UWorld*);
 };
-
-typedef TSharedPtr<FVoiceEngineSteam, ESPMode::ThreadSafe> FVoiceEngineImplPtr;
diff --git a/Engine/Plugins/Online/OnlineSubsystemUtils/Source/OnlineSubsystemUtils/Public/VoiceEngineImpl.h b/Engine/Plugins/Online/OnlineSubsystemUtils/Source/OnlineSubsystemUtils/Public/VoiceEngineImpl.h
--- a/Engine/Plugins/Online/OnlineSubsystemUtils/Source/OnlineSubsystemUtils/Public/VoiceEngineImpl.h
+++ b/Engine/Plugins/Online/OnlineSubsystemUtils/Source/OnlineSubsystemUtils/Public/VoiceEngineImpl.h
@@ -161,6 +161,7 @@ class ONLINESUBSYSTEMUTILS_API FVoiceEngineImpl : public IVoiceEngine, public FS
 	const double DeviceChangeDelay = 2.0f;
 #endif
 
+protected:
 	/**
 	 * Determines if the specified index is the owner or not
 	 *
@@ -168,28 +169,29 @@ class ONLINESUBSYSTEMUTILS_API FVoiceEngineImpl : public IVoiceEngine, public FS
 	 *
 	 * @return true if this is the owner, false otherwise
 	 */
-	FORCEINLINE bool IsOwningUser(uint32 UserIndex)
+	FORCEINLINE virtual bool IsOwningUser(uint32 UserIndex)
 	{
 		return UserIndex >= 0 && UserIndex < MAX_SPLITSCREEN_TALKERS && OwningUserIndex == UserIndex;
 	}
 
-	/** 
-	 * Update the internal state of the voice capturing state 
-	 * Handles possible continuation waiting for capture stop event
-	 */
-	void VoiceCaptureUpdate() const;
-
 	/** Start capturing voice data */
-	void StartRecording() const;
+	virtual void StartRecording() const;
 
 	/** Stop capturing voice data */
-	void StopRecording() const;
+	virtual void StopRecording() const;
 
 	/** Called when "last half second" is over */
-	void StoppedRecording() const;
+	virtual void StoppedRecording() const;
 
 	/** @return is active recording occurring at the moment */
-	bool IsRecording() const { return bIsCapturing || bPendingFinalCapture; }
+	virtual bool IsRecording() const { return bIsCapturing || bPendingFinalCapture; }
+
+private:
+	/**
+	 * Update the internal state of the voice capturing state
+	 * Handles possible continuation waiting for capture stop event
+	 */
+	void VoiceCaptureUpdate() const;
 
 	/**
 	 * Callback from streaming audio when data is requested for playback
@@ -298,15 +300,15 @@ private:
 	void OnPostLoadMap(UWorld*);
 
 protected:
-	IOnlineSubsystem*          GetOnlineSubSystem()         { return OnlineSubsystem; }
-	TSharedPtr<IVoiceCapture>& GetVoiceCapture()            { return VoiceCapture; }
-	TSharedPtr<IVoiceEncoder>& GetVoiceEncoder()            { return VoiceEncoder; }
-	FRemoteTalkerData&         GetRemoteTalkerBuffers()     { return RemoteTalkerBuffers; }
-	TArray<uint8>&             GetCompressedVoiceBuffer()   { return CompressedVoiceBuffer; }
-	TArray<uint8>&             GetDecompressedVoiceBuffer() { return DecompressedVoiceBuffer; }
-	FLocalVoiceData*           GetLocalPlayerVoiceData()    { return PlayerVoiceData; }
-	int32                      GetMaxVoiceRemainderSize();
-	void					   CreateSerializeHelper();
+	virtual IOnlineSubsystem*				 GetOnlineSubSystem()			{ return OnlineSubsystem; }
+	virtual const TSharedPtr<IVoiceCapture>& GetVoiceCapture() const		{ return VoiceCapture; }
+	virtual const TSharedPtr<IVoiceEncoder>& GetVoiceEncoder() const		{ return VoiceEncoder; }
+	virtual FRemoteTalkerData&				 GetRemoteTalkerBuffers()		{ return RemoteTalkerBuffers; }
+	virtual TArray<uint8>&					 GetCompressedVoiceBuffer()		{ return CompressedVoiceBuffer; }
+	virtual TArray<uint8>&					 GetDecompressedVoiceBuffer()	{ return DecompressedVoiceBuffer; }
+	virtual FLocalVoiceData*				 GetLocalPlayerVoiceData()		{ return PlayerVoiceData; }
+	virtual int32							 GetMaxVoiceRemainderSize();
+	virtual void							 CreateSerializeHelper();
 
 	// Get Audio Device Changes on Windows
 #if PLATFORM_WINDOWS
