PR #5842: Implemented FOnlineSessionSteam::CreateSessionIdFromString (Contributed by [REDACTED_PERSON])

diff --git a/Engine/Plugins/Online/OnlineSubsystemSteam/Source/Private/OnlineSessionInterfaceSteam.cpp b/Engine/Plugins/Online/OnlineSubsystemSteam/Source/Private/OnlineSessionInterfaceSteam.cpp
--- a/Engine/Plugins/Online/OnlineSubsystemSteam/Source/Private/OnlineSessionInterfaceSteam.cpp
+++ b/Engine/Plugins/Online/OnlineSubsystemSteam/Source/Private/OnlineSessionInterfaceSteam.cpp
@@ -1341,6 +1341,15 @@ FString FOnlineSessionSteam::GetCustomDedicatedServerName() const
 	return TEXT("");
 }
 
+TSharedPtr<const FUniqueNetId> FOnlineSessionSteam::CreateSessionIdFromString(const FString& SessionIdStr)
+{
+	if (!SessionIdStr.IsEmpty())
+	{
+		return MakeShared<FUniqueNetIdSteam>(SessionIdStr);
+	}
+	return nullptr;
+}
+
 FOnlineSessionSettings* FOnlineSessionSteam::GetSessionSettings(FName SessionName) 
 {
 	FNamedOnlineSession* Session = GetNamedSession(SessionName);
diff --git a/Engine/Plugins/Online/OnlineSubsystemSteam/Source/Private/OnlineSessionInterfaceSteam.h b/Engine/Plugins/Online/OnlineSubsystemSteam/Source/Private/OnlineSessionInterfaceSteam.h
--- a/Engine/Plugins/Online/OnlineSubsystemSteam/Source/Private/OnlineSessionInterfaceSteam.h
+++ b/Engine/Plugins/Online/OnlineSubsystemSteam/Source/Private/OnlineSessionInterfaceSteam.h
@@ -447,12 +447,7 @@ public:
 
 	virtual ~FOnlineSessionSteam() {}
 
-	virtual TSharedPtr<const FUniqueNetId> CreateSessionIdFromString(const FString& SessionIdStr) override
-	{
-		/* NYI */
-		TSharedPtr<const FUniqueNetId> SessionId;
-		return SessionId;
-	}
+	virtual TSharedPtr<const FUniqueNetId> CreateSessionIdFromString(const FString& SessionIdStr) override;
 
 	FNamedOnlineSession* GetNamedSession(FName SessionName) override
 	{
