Fix matchmaking + joining issues when the socket subsystem is not SteamNet.

diff --git a/Engine/Plugins/Online/OnlineSubsystemSteam/Source/Private/OnlineSessionAsyncServerSteam.cpp b/Engine/Plugins/Online/OnlineSubsystemSteam/Source/Private/OnlineSessionAsyncServerSteam.cpp
--- a/Engine/Plugins/Online/OnlineSubsystemSteam/Source/Private/OnlineSessionAsyncServerSteam.cpp
+++ b/Engine/Plugins/Online/OnlineSubsystemSteam/Source/Private/OnlineSessionAsyncServerSteam.cpp
@@ -677,7 +677,8 @@ bool FPendingSearchResultSteam::FillSessionFromServerRules()
 	}
 
 	// Verify success with all required keys found
-	if (bSuccess && (KeysFound == STEAMKEY_NUMREQUIREDSERVERKEYS) && (SteamAddrKeysFound == 2))
+	// If the user has SteamNetworking off, then we should just check if their host addressing is correct.
+	if (bSuccess && (KeysFound == STEAMKEY_NUMREQUIREDSERVERKEYS) && (SteamAddrKeysFound == 2 || (HostAddr->IsValid() && HostAddr.IsValid())))
 	{
 		SessionInfo->HostAddr = HostAddr;
 		SessionInfo->SteamP2PAddr = SteamP2PAddr;
@@ -1275,10 +1276,8 @@ void FOnlineAsyncEventSteamInviteAccepted::Finalize()
 		Port = (Port > 0) ? Port : Subsystem->GetGameServerGamePort();
 
 		// Parse the address
-		bool bIsValid;
-		TSharedRef<FInternetAddr> IpAddr = ISocketSubsystem::Get(PLATFORM_SOCKETSUBSYSTEM)->CreateInternetAddr();
-		IpAddr->SetIp(ParsedURL, bIsValid);
-		if (bIsValid)
+		TSharedPtr<FInternetAddr> IpAddr = ISocketSubsystem::Get(PLATFORM_SOCKETSUBSYSTEM)->GetAddressFromString(ParsedURL);
+		if (IpAddr.IsValid())
 		{
 			SessionInt->CurrentSessionSearch->QuerySettings.Set(FName(SEARCH_STEAM_HOSTIP), IpAddr->ToString(false), EOnlineComparisonOp::Equals);
 			FOnlineAsyncTaskSteamFindServerForInviteSession* NewTask = new FOnlineAsyncTaskSteamFindServerForInviteSession(Subsystem, SearchSettings, LocalUserNum, SessionInt->OnSessionUserInviteAcceptedDelegates);
diff --git a/Engine/Plugins/Online/OnlineSubsystemSteam/Source/Private/OnlineSessionInterfaceSteam.cpp b/Engine/Plugins/Online/OnlineSubsystemSteam/Source/Private/OnlineSessionInterfaceSteam.cpp
--- a/Engine/Plugins/Online/OnlineSubsystemSteam/Source/Private/OnlineSessionInterfaceSteam.cpp
+++ b/Engine/Plugins/Online/OnlineSubsystemSteam/Source/Private/OnlineSessionInterfaceSteam.cpp
@@ -1228,35 +1228,27 @@ static bool GetConnectStringFromSessionInfo(TSharedPtr<FOnlineSessionInfoSteam>&
 
 	if (SessionInfo.IsValid())
 	{
-		if (SessionInfo->SessionType == ESteamSession::LobbySession ||
-			SessionInfo->SessionType == ESteamSession::AdvertisedSessionHost ||
-			SessionInfo->SessionType == ESteamSession::AdvertisedSessionClient)
+		if (SessionInfo->SteamP2PAddr.IsValid() && SessionInfo->SteamP2PAddr->IsValid())
 		{
-			if (SessionInfo->SteamP2PAddr.IsValid() && SessionInfo->SteamP2PAddr->IsValid())
+			int32 SteamPort = SessionInfo->SteamP2PAddr->GetPort();
+			if (PortOverride > 0)
 			{
-				int32 SteamPort = SessionInfo->SteamP2PAddr->GetPort();
-				if (PortOverride > 0)
-				{
-					SteamPort = PortOverride;
-				}
-
-				ConnectInfo = FString::Printf(STEAM_URL_PREFIX TEXT("%s:%d"), *SessionInfo->SteamP2PAddr->ToString(false), SteamPort);
-				bSuccess = true;
+				SteamPort = PortOverride;
 			}
+
+			ConnectInfo = FString::Printf(STEAM_URL_PREFIX TEXT("%s:%d"), *SessionInfo->SteamP2PAddr->ToString(false), SteamPort);
+			bSuccess = true;
 		}
-		else
+		else if (SessionInfo->HostAddr.IsValid() && SessionInfo->HostAddr->IsValid())
 		{
-			if (SessionInfo->HostAddr.IsValid() && SessionInfo->HostAddr->IsValid())
+			int32 HostPort = SessionInfo->HostAddr->GetPort();
+			if (PortOverride > 0)
 			{
-				int32 HostPort = SessionInfo->HostAddr->GetPort();
-				if (PortOverride > 0)
-				{
-					HostPort = PortOverride;
-				}
-
-				ConnectInfo = FString::Printf(TEXT("%s:%d"), *SessionInfo->HostAddr->ToString(false), HostPort);
-				bSuccess = true;
+				HostPort = PortOverride;
 			}
+
+			ConnectInfo = FString::Printf(TEXT("%s:%d"), *SessionInfo->HostAddr->ToString(false), HostPort);
+			bSuccess = true;
 		}
 	}
 
diff --git a/Engine/Plugins/Online/OnlineSubsystemSteam/Source/Private/OnlineSubsystemSteamTypes.h b/Engine/Plugins/Online/OnlineSubsystemSteam/Source/Private/OnlineSubsystemSteamTypes.h
--- a/Engine/Plugins/Online/OnlineSubsystemSteam/Source/Private/OnlineSubsystemSteamTypes.h
+++ b/Engine/Plugins/Online/OnlineSubsystemSteam/Source/Private/OnlineSubsystemSteamTypes.h
@@ -322,8 +322,7 @@ public:
 			return SteamP2PAddr.IsValid() && SteamP2PAddr->IsValid() && SessionId.IsValid();
 		case ESteamSession::AdvertisedSessionHost:
 		case ESteamSession::AdvertisedSessionClient:
-			// Could/should check that the HostAddr is valid here also
-			return SteamP2PAddr.IsValid() && SteamP2PAddr->IsValid() && SessionId.IsValid();
+			return ((SteamP2PAddr.IsValid() && SteamP2PAddr->IsValid()) || (HostAddr.IsValid() && HostAddr->IsValid())) && SessionId.IsValid();
 		case ESteamSession::LANSession:
 		default:
 			// LAN case
