Fixing build error caused by outdated interface methods.

diff --git a/Engine/Plugins/Online/IOS/OnlineSubsystemIOS/Source/Private/OnlineUserCloudInterfaceIOS.cpp b/Engine/Plugins/Online/IOS/OnlineSubsystemIOS/Source/Private/OnlineUserCloudInterfaceIOS.cpp
--- a/Engine/Plugins/Online/IOS/OnlineSubsystemIOS/Source/Private/OnlineUserCloudInterfaceIOS.cpp
+++ b/Engine/Plugins/Online/IOS/OnlineSubsystemIOS/Source/Private/OnlineUserCloudInterfaceIOS.cpp
@@ -364,7 +364,7 @@ bool FOnlineUserCloudInterfaceIOS::ReadUserFile(const FUniqueNetId& UserId, cons
 	return false;
 }
 
-bool FOnlineUserCloudInterfaceIOS::WriteUserFile(const FUniqueNetId& UserId, const FString& FileName, TArray<uint8>& FileContents)
+bool FOnlineUserCloudInterfaceIOS::WriteUserFile(const FUniqueNetId& UserId, const FString& FileName, TArray<uint8>& FileContents, bool bCompressBeforeUpload)
 {
 #ifdef __IPHONE_8_0
 	FCloudFile* CloudFile = GetCloudFile(FileName, true);
diff --git a/Engine/Plugins/Online/IOS/OnlineSubsystemIOS/Source/Private/OnlineUserCloudInterfaceIOS.h b/Engine/Plugins/Online/IOS/OnlineSubsystemIOS/Source/Private/OnlineUserCloudInterfaceIOS.h
--- a/Engine/Plugins/Online/IOS/OnlineSubsystemIOS/Source/Private/OnlineUserCloudInterfaceIOS.h
+++ b/Engine/Plugins/Online/IOS/OnlineSubsystemIOS/Source/Private/OnlineUserCloudInterfaceIOS.h
@@ -54,7 +54,7 @@ public:
 	virtual void EnumerateUserFiles(const FUniqueNetId& UserId) override;
 	virtual void GetUserFileList(const FUniqueNetId& UserId, TArray<FCloudFileHeader>& UserFiles) override;
 	virtual bool ReadUserFile(const FUniqueNetId& UserId, const FString& FileName) override;
-	virtual bool WriteUserFile(const FUniqueNetId& UserId, const FString& FileName, TArray<uint8>& FileContents) override;
+	virtual bool WriteUserFile(const FUniqueNetId& UserId, const FString& FileName, TArray<uint8>& FileContents, bool bCompressBeforeUpload = false) override;
 	virtual void CancelWriteUserFile(const FUniqueNetId& UserId, const FString& FileName) override;
 	virtual bool DeleteUserFile(const FUniqueNetId& UserId, const FString& FileName, bool bShouldCloudDelete, bool bShouldLocallyDelete) override;
 	virtual bool RequestUsageInfo(const FUniqueNetId& UserId) override;
diff --git a/Engine/Plugins/Online/OnlineSubsystemOculus/Source/Private/OnlineUserCloudOculus.cpp b/Engine/Plugins/Online/OnlineSubsystemOculus/Source/Private/OnlineUserCloudOculus.cpp
--- a/Engine/Plugins/Online/OnlineSubsystemOculus/Source/Private/OnlineUserCloudOculus.cpp
+++ b/Engine/Plugins/Online/OnlineSubsystemOculus/Source/Private/OnlineUserCloudOculus.cpp
@@ -135,7 +135,7 @@ void FOnlineUserCloudOculus::GetUserFileList(const FUniqueNetId& UserId, TArray<
 	UserFiles = MoveTemp(EnumerateCache);
 }
 
-bool FOnlineUserCloudOculus::WriteUserFile(const FUniqueNetId& UserId, const FString& FileName, TArray<uint8>& FileContents)
+bool FOnlineUserCloudOculus::WriteUserFile(const FUniqueNetId& UserId, const FString& FileName, TArray<uint8>& FileContents, bool bCompressBeforeUpload)
 {
 	auto LoggedInPlayerId = OculusSubsystem.GetIdentityInterface()->GetUniquePlayerId(0);
 	if (!LoggedInPlayerId.IsValid() || UserId != *LoggedInPlayerId)
diff --git a/Engine/Plugins/Online/OnlineSubsystemOculus/Source/Private/OnlineUserCloudOculus.h b/Engine/Plugins/Online/OnlineSubsystemOculus/Source/Private/OnlineUserCloudOculus.h
--- a/Engine/Plugins/Online/OnlineSubsystemOculus/Source/Private/OnlineUserCloudOculus.h
+++ b/Engine/Plugins/Online/OnlineSubsystemOculus/Source/Private/OnlineUserCloudOculus.h
@@ -35,7 +35,7 @@ public:
 
 	virtual bool ClearFile(const FUniqueNetId& UserId, const FString& FileName) override;
 
-	virtual bool WriteUserFile(const FUniqueNetId& UserId, const FString& FileName, TArray<uint8>& FileContents) override;
+	virtual bool WriteUserFile(const FUniqueNetId& UserId, const FString& FileName, TArray<uint8>& FileContents, bool bCompressBeforeUpload = false) override;
 
 	virtual void CancelWriteUserFile(const FUniqueNetId& UserId, const FString& FileName) override;
 
diff --git a/Engine/Plugins/Online/OnlineSubsystemSteam/Source/Private/OnlineUserCloudInterfaceSteam.cpp b/Engine/Plugins/Online/OnlineSubsystemSteam/Source/Private/OnlineUserCloudInterfaceSteam.cpp
--- a/Engine/Plugins/Online/OnlineSubsystemSteam/Source/Private/OnlineUserCloudInterfaceSteam.cpp
+++ b/Engine/Plugins/Online/OnlineSubsystemSteam/Source/Private/OnlineUserCloudInterfaceSteam.cpp
@@ -316,7 +316,7 @@ bool FOnlineUserCloudSteam::ReadUserFile(const FUniqueNetId& UserId, const FStri
 	return false;
 }
 
-bool FOnlineUserCloudSteam::WriteUserFile(const FUniqueNetId& UserId, const FString& FileName, TArray<uint8>& FileContents)
+bool FOnlineUserCloudSteam::WriteUserFile(const FUniqueNetId& UserId, const FString& FileName, TArray<uint8>& FileContents, bool bCompressBeforeUpload)
 {
 	FScopeLock ScopeLock(&SteamSubsystem->UserCloudDataLock);
 	// Create or get the current entry for this file
diff --git a/Engine/Plugins/Online/OnlineSubsystemSteam/Source/Private/OnlineUserCloudInterfaceSteam.h b/Engine/Plugins/Online/OnlineSubsystemSteam/Source/Private/OnlineUserCloudInterfaceSteam.h
--- a/Engine/Plugins/Online/OnlineSubsystemSteam/Source/Private/OnlineUserCloudInterfaceSteam.h
+++ b/Engine/Plugins/Online/OnlineSubsystemSteam/Source/Private/OnlineUserCloudInterfaceSteam.h
@@ -232,7 +232,7 @@ public:
 	virtual void EnumerateUserFiles(const FUniqueNetId& UserId) override;
 	virtual void GetUserFileList(const FUniqueNetId& UserId, TArray<FCloudFileHeader>& UserFiles) override;
 	virtual bool ReadUserFile(const FUniqueNetId& UserId, const FString& FileName) override;
-	virtual bool WriteUserFile(const FUniqueNetId& UserId, const FString& FileName, TArray<uint8>& FileContents) override;
+	virtual bool WriteUserFile(const FUniqueNetId& UserId, const FString& FileName, TArray<uint8>& FileContents, bool bCompressBeforeUpload = false) override;
 	virtual void CancelWriteUserFile(const FUniqueNetId& UserId, const FString& FileName) override;
 	virtual bool DeleteUserFile(const FUniqueNetId& UserId, const FString& FileName, bool bShouldCloudDelete, bool bShouldLocallyDelete) override;
 	virtual bool RequestUsageInfo(const FUniqueNetId& UserId) override;
