Another fix for a wrongly-changed UE_ARRAY_COUNT.

diff --git a/Engine/Plugins/Online/OnlineSubsystemSteam/Source/Private/OnlineSubsystemSteamPrivate.h b/Engine/Plugins/Online/OnlineSubsystemSteam/Source/Private/OnlineSubsystemSteamPrivate.h
--- a/Engine/Plugins/Online/OnlineSubsystemSteam/Source/Private/OnlineSubsystemSteamPrivate.h
+++ b/Engine/Plugins/Online/OnlineSubsystemSteam/Source/Private/OnlineSubsystemSteamPrivate.h
@@ -20,8 +20,9 @@
 #undef ONLINE_LOG_PREFIX
 #define ONLINE_LOG_PREFIX TEXT("STEAM: ")
 
-#pragma push_macro("UE_ARRAY_COUNT")
-#undef UE_ARRAY_COUNT
+// These can be removed when the deprecated ARRAY_COUNT is removed.
+#pragma push_macro("ARRAY_COUNT")
+#undef ARRAY_COUNT
 
 // Steamworks SDK headers
 #if STEAMSDK_FOUND == 0
@@ -35,4 +36,4 @@ THIRD_PARTY_INCLUDES_START
 
 THIRD_PARTY_INCLUDES_END
 
-#pragma pop_macro("UE_ARRAY_COUNT")
+#pragma pop_macro("ARRAY_COUNT")
