Fixing SteamSockets matchmaking issues due to data ambiguity from the platform response types.

#rnx

diff --git a/Engine/Plugins/Online/OnlineSubsystemSteam/Source/Private/OnlineSessionAsyncLobbySteam.cpp b/Engine/Plugins/Online/OnlineSubsystemSteam/Source/Private/OnlineSessionAsyncLobbySteam.cpp
--- a/Engine/Plugins/Online/OnlineSubsystemSteam/Source/Private/OnlineSessionAsyncLobbySteam.cpp
+++ b/Engine/Plugins/Online/OnlineSubsystemSteam/Source/Private/OnlineSessionAsyncLobbySteam.cpp
@@ -154,6 +154,7 @@ static void GetLobbyKeyValuePairsFromSessionSettings(const FOnlineSessionSetting
  */
 static void GetLobbyKeyValuePairsFromSessionInfo(const FOnlineSessionInfoSteam& SessionInfo, FSteamSessionKeyValuePairs& KeyValuePairs)
 {
+	// This should pretty much always be false.
 	if (SessionInfo.HostAddr.IsValid())
 	{
 		uint32 HostAddr;
@@ -162,11 +163,13 @@ static void GetLobbyKeyValuePairsFromSessionInfo(const FOnlineSessionInfoSteam&
 		KeyValuePairs.Add(STEAMKEY_HOSTPORT, FString::FromInt(SessionInfo.HostAddr->GetPort()));
 	}
 
+	// Lobbies are exclusively P2P.
+	KeyValuePairs.Add(STEAMKEY_CONNECTIONMETHOD, FString::Printf(TEXT("%s"), *LexToString(FSteamConnectionMethod::P2P)));
+
 	if (SessionInfo.SteamP2PAddr.IsValid())
 	{
-		TSharedPtr<FInternetAddrSteam> SteamAddr = StaticCastSharedPtr<FInternetAddrSteam>(SessionInfo.SteamP2PAddr);
-		KeyValuePairs.Add(STEAMKEY_P2PADDR, SteamAddr->ToString(false));
-		KeyValuePairs.Add(STEAMKEY_P2PPORT, FString::FromInt(SteamAddr->GetPort()));
+		KeyValuePairs.Add(STEAMKEY_P2PADDR, SessionInfo.SteamP2PAddr->ToString(false));
+		KeyValuePairs.Add(STEAMKEY_P2PPORT, FString::FromInt(SessionInfo.SteamP2PAddr->GetPort()));
 	}
 }
 
@@ -342,9 +345,18 @@ bool FillSessionFromLobbyData(FOnlineSubsystemSteam* SteamSubsystem, FUniqueNetI
 				HostKeysFound++;
 			}
 		}
+		else if (FCStringAnsi::Stricmp(Key, STEAMKEY_CONNECTIONMETHOD) == 0)
+		{
+			SessionInfo->ConnectionMethod = ToConnectionMethod(ANSI_TO_TCHAR(Value));
+			++KeysFound;
+		}
 		else if (FCStringAnsi::Stricmp(Key, STEAMKEY_P2PADDR) == 0)
 		{
-			uint64 SteamAddr = FCString::Atoi64(ANSI_TO_TCHAR(Value));
+			FString KeyValue = ANSI_TO_TCHAR(Value);
+			// Remove any protocol flags from the start.
+			KeyValue.RemoveFromStart(STEAM_URL_PREFIX);
+
+			uint64 SteamAddr = FCString::Atoi64(*KeyValue);
 			if (SteamAddr != 0)
 			{
 				SteamP2PAddr->SteamId.UniqueNetId = SteamAddr;
diff --git a/Engine/Plugins/Online/OnlineSubsystemSteam/Source/Private/OnlineSessionAsyncServerSteam.cpp b/Engine/Plugins/Online/OnlineSubsystemSteam/Source/Private/OnlineSessionAsyncServerSteam.cpp
--- a/Engine/Plugins/Online/OnlineSubsystemSteam/Source/Private/OnlineSessionAsyncServerSteam.cpp
+++ b/Engine/Plugins/Online/OnlineSubsystemSteam/Source/Private/OnlineSessionAsyncServerSteam.cpp
@@ -3,11 +3,14 @@
 #include "OnlineSessionAsyncServerSteam.h"
 #include "GameFramework/GameStateBase.h"
 #include "GameFramework/PlayerState.h"
+#include "Misc/CoreMisc.h"
+#include "Misc/ConfigCacheIni.h"
 #include "OnlineSubsystemUtils.h"
 #include "SocketSubsystem.h"
 #include "IPAddressSteam.h"
 #include "SteamSessionKeys.h"
 #include "SteamUtilities.h"
+#include "OnlineSubsystemSteam.h"
 #include "OnlineAuthInterfaceSteam.h"
 
 
@@ -128,24 +131,58 @@ void GetServerKeyValuePairsFromSessionSettings(const FOnlineSessionSettings& Ses
  * @param SessionInfo session info to get key, value pairs from
  * @param KeyValuePairs key value pair structure to add to
  */
-void GetServerKeyValuePairsFromSessionInfo(const FOnlineSessionInfoSteam* SessionInfo, FSteamSessionKeyValuePairs& KeyValuePairs)
+void GetServerKeyValuePairsFromSessionInfo(FOnlineSessionInfoSteam* SessionInfo, FSteamSessionKeyValuePairs& KeyValuePairs)
 {
-	// @TODO ONLINE - not needed
-	/*
-	if (SessionInfo->HostAddr.IsValid())
+	FOnlineSubsystemSteam* SteamSubsystem = static_cast<FOnlineSubsystemSteam*>(IOnlineSubsystem::Get(STEAM_SUBSYSTEM));
+	FSteamConnectionMethod MethodUsed = SessionInfo->ConnectionMethod;
+
+	// We should not enter this check here if we're using LAN.
+	if (SteamSubsystem && GConfig && MethodUsed == FSteamConnectionMethod::None)
 	{
-		uint32 HostAddr;
-		SessionInfo->HostAddr->GetIp(HostAddr);
-		KeyValuePairs.Add(STEAMKEY_HOSTIP, FString::FromInt(HostAddr));
-		KeyValuePairs.Add(STEAMKEY_HOSTPORT, FString::FromInt(SessionInfo->HostAddr->GetPort()));
+		bool bUseRelays = false;
+		bool bIsDefaultSubsystem = false;
+		GConfig->GetBool(TEXT("OnlineSubsystemSteam"), TEXT("bAllowP2PPacketRelay"), bUseRelays, GEngineIni);
+		// If this config specifically does not exist, then set it to true
+		if (!GConfig->GetBool(TEXT("OnlineSubsystemSteam"), TEXT("bUseSteamNetworking"), bIsDefaultSubsystem, GEngineIni))
+		{
+			bIsDefaultSubsystem = true;
+		}
+
+		// Using Steam sockets. This checks if the SteamSockets plugin is enabled
+		if (!SteamSubsystem->IsUsingSteamNetworking())
+		{
+			if (!bUseRelays)
+			{
+				MethodUsed = FSteamConnectionMethod::Direct;
+			}
+			else
+			{
+				MethodUsed = FSteamConnectionMethod::P2P;
+			}
+		}
+		else
+		{
+			// Determine if we are a default. If so, we're using SteamNetworking P2P.
+			if (bIsDefaultSubsystem)
+			{
+				MethodUsed = FSteamConnectionMethod::P2P;
+			}
+			else
+			{
+				MethodUsed = FSteamConnectionMethod::Direct;
+			}
+		}
+
+		// Set this if it hasn't been set yet.
+		SessionInfo->ConnectionMethod = MethodUsed;
 	}
-	*/
+
+	KeyValuePairs.Add(STEAMKEY_CONNECTIONMETHOD, FString::Printf(TEXT("%s"), *LexToString(MethodUsed)));
 
 	if (SessionInfo->SteamP2PAddr.IsValid())
 	{
-		TSharedPtr<FInternetAddrSteam> SteamAddr = StaticCastSharedPtr<FInternetAddrSteam>(SessionInfo->SteamP2PAddr);
-		KeyValuePairs.Add(STEAMKEY_P2PADDR, SteamAddr->ToString(false));
-		KeyValuePairs.Add(STEAMKEY_P2PPORT, FString::FromInt(SteamAddr->GetPort()));
+		KeyValuePairs.Add(STEAMKEY_P2PADDR, SessionInfo->SteamP2PAddr->ToString(false));
+		KeyValuePairs.Add(STEAMKEY_P2PPORT, FString::FromInt(SessionInfo->SteamP2PAddr->GetPort()));
 	}
 }
 
@@ -635,6 +672,11 @@ bool FPendingSearchResultSteam::FillSessionFromServerRules()
 				KeysFound++;
 			}
 		}
+		else if (FCStringAnsi::Stricmp(TCHAR_TO_ANSI(*It.Key()), STEAMKEY_CONNECTIONMETHOD) == 0)
+		{
+			SessionInfo->ConnectionMethod = ToConnectionMethod(It.Value());
+			++KeysFound;
+		}
 // 		else if (FCStringAnsi::Stricmp(TCHAR_TO_ANSI(*It.Key()), STEAMKEY_NUMOPENPRIVATECONNECTIONS) == 0)
 // 		{
 // 			Session->NumOpenPrivateConnections = FCString::Atoi(*It.Value());
@@ -647,7 +689,11 @@ bool FPendingSearchResultSteam::FillSessionFromServerRules()
 // 		}
 		else if (FCStringAnsi::Stricmp(TCHAR_TO_ANSI(*It.Key()), STEAMKEY_P2PADDR) == 0)
 		{
-			uint64 SteamAddr = FCString::Atoi64(*It.Value());
+			FString KeyValue = It.Value();
+			// Remove any protocol flags from the start.
+			KeyValue.RemoveFromStart(STEAM_URL_PREFIX);
+
+			uint64 SteamAddr = FCString::Atoi64(*KeyValue);
 			if (SteamAddr != 0)
 			{
 				SteamP2PAddr->SteamId.UniqueNetId = SteamAddr;
@@ -678,7 +724,7 @@ bool FPendingSearchResultSteam::FillSessionFromServerRules()
 
 	// Verify success with all required keys found
 	// If the user has SteamNetworking off, then we should just check if their host addressing is correct.
-	if (bSuccess && (KeysFound == STEAMKEY_NUMREQUIREDSERVERKEYS) && (SteamAddrKeysFound == 2 || (HostAddr->IsValid() && HostAddr.IsValid())))
+	if (bSuccess && (KeysFound == STEAMKEY_NUMREQUIREDSERVERKEYS) && (SteamAddrKeysFound == 2 || (HostAddr.IsValid() && HostAddr->IsValid())))
 	{
 		SessionInfo->HostAddr = HostAddr;
 		SessionInfo->SteamP2PAddr = SteamP2PAddr;
diff --git a/Engine/Plugins/Online/OnlineSubsystemSteam/Source/Private/OnlineSessionInterfaceSteam.cpp b/Engine/Plugins/Online/OnlineSubsystemSteam/Source/Private/OnlineSessionInterfaceSteam.cpp
--- a/Engine/Plugins/Online/OnlineSubsystemSteam/Source/Private/OnlineSessionInterfaceSteam.cpp
+++ b/Engine/Plugins/Online/OnlineSubsystemSteam/Source/Private/OnlineSessionInterfaceSteam.cpp
@@ -18,18 +18,20 @@
 /** Constructor for non-lobby sessions */
 FOnlineSessionInfoSteam::FOnlineSessionInfoSteam(ESteamSession::Type InSessionType) :
 	SessionType(InSessionType),
-	HostAddr(NULL),
-	SteamP2PAddr(NULL),
-	SessionId((uint64)0)
+	HostAddr(nullptr),
+	SteamP2PAddr(nullptr),
+	SessionId((uint64)0),
+	ConnectionMethod((InSessionType == ESteamSession::LANSession) ? FSteamConnectionMethod::Direct : FSteamConnectionMethod::None)
 {
 }
 
 /** Constructor for sessions that represent a Steam lobby or server */
 FOnlineSessionInfoSteam::FOnlineSessionInfoSteam(ESteamSession::Type InSessionType, const FUniqueNetIdSteam& InSessionId) :
 	SessionType(InSessionType),
-	HostAddr(NULL),
-	SteamP2PAddr(NULL),
-	SessionId(InSessionId)
+	HostAddr(nullptr),
+	SteamP2PAddr(nullptr),
+	SessionId(InSessionId),
+	ConnectionMethod(FSteamConnectionMethod::None)
 {
 }
 
@@ -40,6 +42,7 @@ void FOnlineSessionInfoSteam::Init()
 void FOnlineSessionInfoSteam::InitLAN()
 {
 	SessionType = ESteamSession::LANSession;
+	ConnectionMethod = FSteamConnectionMethod::Direct;
 
 	uint64 Nonce = 0;
 	GenerateNonce((uint8*)&Nonce, 8);
@@ -218,7 +221,7 @@ bool FOnlineSessionSteam::CreateSession(int32 HostingPlayerNum, FName SessionNam
 
 	// Check for an existing session
 	FNamedOnlineSession* Session = GetNamedSession(SessionName);
-	if (Session == NULL)
+	if (Session == nullptr)
 	{
 		// Create a new session and deep copy the game settings
 		Session = AddNamedSession(SessionName, NewSessionSettings);
@@ -317,7 +320,7 @@ uint32 FOnlineSessionSteam::CreateInternetSession(int32 HostingPlayerNum, FNamed
 
 	// Only allowed one published session with Steam
 	FNamedOnlineSession* MasterSession = GetGameServerSession();
-	if (MasterSession == NULL)
+	if (MasterSession == nullptr)
 	{
 		if (SteamSubsystem->IsSteamServerAvailable())
 		{
@@ -392,7 +395,7 @@ bool FOnlineSessionSteam::StartSession(FName SessionName)
 				Result = ONLINE_SUCCESS;
 				Session->SessionState = EOnlineSessionState::InProgress;
 
-				if (SteamFriends() != NULL)
+				if (SteamFriends() != nullptr)
 				{
 					for (int32 PlayerIdx=0; PlayerIdx < Session->RegisteredPlayers.Num(); PlayerIdx++)
 					{
@@ -494,7 +497,7 @@ bool FOnlineSessionSteam::EndSession(FName SessionName)
 			{
 				// If the session should be advertised and the lan beacon was destroyed, recreate
 				if (Session->SessionSettings.bShouldAdvertise && 
-					LANSession->LanBeacon == NULL &&
+					LANSession->LanBeacon == nullptr &&
 					SteamSubsystem->IsServer())
 				{
 					// Recreate the beacon
@@ -586,7 +589,7 @@ bool FOnlineSessionSteam::DestroySession(FName SessionName, const FOnDestroySess
 					// Tear down the LAN beacon
 					LANSession->StopLANSession();
 					delete LANSession;
-					LANSession = NULL;
+					LANSession = nullptr;
 				}
 
 				Result = ONLINE_SUCCESS;
@@ -647,7 +650,7 @@ uint32 FOnlineSessionSteam::DestroyInternetSession(FNamedOnlineSession* Session,
 
 	// Clear any session advertisements this account had with this session.
 	FOnlineAuthSteamPtr SteamAuth = SteamSubsystem->GetAuthInterface();
-	if (SteamUser() != NULL && SteamUser()->BLoggedOn() && (!SteamAuth.IsValid() || !SteamAuth->IsSessionAuthEnabled()))
+	if (SteamUser() != nullptr && SteamUser()->BLoggedOn() && (!SteamAuth.IsValid() || !SteamAuth->IsSessionAuthEnabled()))
 	{
 		UE_LOG_ONLINE(Warning, TEXT("AUTH: DestroyInternetSession is calling the depricated AdvertiseGame call"));
 		SteamUser()->AdvertiseGame(k_steamIDNil, 0, 0);
@@ -792,7 +795,7 @@ uint32 FOnlineSessionSteam::FindLANSession(const TSharedRef<FOnlineSessionSearch
 		{
 			Return = ONLINE_FAIL;
 			delete LANSession;
-			LANSession = NULL;
+			LANSession = nullptr;
 
 			CurrentSessionSearch->SearchState = EOnlineAsyncTaskState::Failed;
 
@@ -828,7 +831,7 @@ bool FOnlineSessionSteam::CancelFindSessions()
 			// There is no CANCEL lobby query
 			// NULLing out the object will prevent the async event from adding the results
 			CurrentSessionSearch->SearchState = EOnlineAsyncTaskState::Failed;
-			CurrentSessionSearch = NULL;
+			CurrentSessionSearch = nullptr;
 		}
 	}
 	else
@@ -849,7 +852,7 @@ bool FOnlineSessionSteam::JoinSession(int32 PlayerNum, FName SessionName, const
 	uint32 Return = ONLINE_FAIL;
 	FNamedOnlineSession* Session = GetNamedSession(SessionName);
 	// Don't join a session if already in one or hosting one
-	if (Session == NULL)
+	if (Session == nullptr)
 	{
 		// Create a named session from the search result data
 		Session = AddNamedSession(SessionName, DesiredSession.Session);
@@ -935,6 +938,7 @@ uint32 FOnlineSessionSteam::JoinLobbySession(int32 PlayerNum, FNamedOnlineSessio
 			const FOnlineSessionInfoSteam* SearchSessionInfo = (const FOnlineSessionInfoSteam*)SearchSession->SessionInfo.Get();
 			SteamSessionInfo->HostAddr = SearchSessionInfo->HostAddr;
 			SteamSessionInfo->SteamP2PAddr = SearchSessionInfo->SteamP2PAddr;
+			SteamSessionInfo->ConnectionMethod = SearchSessionInfo->ConnectionMethod;
 
 			// The settings found on the search object will be duplicated again when we enter the lobby, possibly updated
 			FOnlineAsyncTaskSteamJoinLobby* NewTask = new FOnlineAsyncTaskSteamJoinLobby(SteamSubsystem, Session->SessionName, SteamSessionInfo->SessionId);
@@ -960,6 +964,7 @@ uint32 FOnlineSessionSteam::JoinInternetSession(int32 PlayerNum, FNamedOnlineSes
 			const FOnlineSessionInfoSteam* SearchSessionInfo = (const FOnlineSessionInfoSteam*)SearchSession->SessionInfo.Get();
 			SteamSessionInfo->HostAddr = SearchSessionInfo->HostAddr;
 			SteamSessionInfo->SteamP2PAddr = SearchSessionInfo->SteamP2PAddr;
+			SteamSessionInfo->ConnectionMethod = SearchSessionInfo->ConnectionMethod;
 
 			if (SearchSession->SessionSettings.bAllowJoinViaPresence)
 			{
@@ -1006,6 +1011,7 @@ uint32 FOnlineSessionSteam::JoinLANSession(int32 PlayerNum, FNamedOnlineSession*
 		FOnlineSessionInfoSteam* SessionInfo = (FOnlineSessionInfoSteam*)Session->SessionInfo.Get();
 
 		SessionInfo->HostAddr = SearchSessionInfo->HostAddr->Clone();
+		SessionInfo->ConnectionMethod = FSteamConnectionMethod::Direct;
 		Result = ONLINE_SUCCESS;
 	}
 
@@ -1099,7 +1105,7 @@ void FOnlineSessionSteam::CheckPendingSessionInvite()
 	{
 		const TCHAR* Str = CmdLine + ConnectIdx + LobbyConnectCmd.Len();
 		FString LobbyIdStr = FParse::Token(Str, 0);
-		int64 LobbyId = FCString::Strtoui64(*LobbyIdStr, NULL, 10);
+		int64 LobbyId = FCString::Strtoui64(*LobbyIdStr, nullptr, 10);
 		if (LobbyId > 0)
 		{
 			PendingInvite.PendingInviteType = ESteamSession::LobbySession;
@@ -1230,28 +1236,42 @@ static bool GetConnectStringFromSessionInfo(TSharedPtr<FOnlineSessionInfoSteam>&
 
 	if (SessionInfo.IsValid())
 	{
-		if (SessionInfo->SteamP2PAddr.IsValid() && SessionInfo->SteamP2PAddr->IsValid())
+		bool bP2PDataValid = (SessionInfo->SteamP2PAddr.IsValid() && SessionInfo->SteamP2PAddr->IsValid());
+		bool bHostDataValid = (SessionInfo->HostAddr.IsValid() && SessionInfo->HostAddr->IsValid());
+
+		// If we have host data, attempt to use it.
+		if (bHostDataValid && SessionInfo->ConnectionMethod == FSteamConnectionMethod::Direct)
 		{
-			int32 SteamPort = SessionInfo->SteamP2PAddr->GetPort();
+			UE_LOG_ONLINE_SESSION(Log, TEXT("Using Host Data for Connection Serialization"));
+
+			int32 HostPort = SessionInfo->HostAddr->GetPort();
 			if (PortOverride > 0)
 			{
-				SteamPort = PortOverride;
+				HostPort = PortOverride;
 			}
 
-			ConnectInfo = FString::Printf(STEAM_URL_PREFIX TEXT("%s:%d"), *SessionInfo->SteamP2PAddr->ToString(false), SteamPort);
+			ConnectInfo = FString::Printf(TEXT("%s:%d"), *SessionInfo->HostAddr->ToString(false), HostPort);
 			bSuccess = true;
 		}
-		else if (SessionInfo->HostAddr.IsValid() && SessionInfo->HostAddr->IsValid())
+		else if (bP2PDataValid)
 		{
-			int32 HostPort = SessionInfo->HostAddr->GetPort();
+			UE_LOG_ONLINE_SESSION(Log, TEXT("Using P2P Data for Connection Serialization"));
+
+			int32 SteamPort = SessionInfo->SteamP2PAddr->GetPort();
 			if (PortOverride > 0)
 			{
-				HostPort = PortOverride;
+				SteamPort = PortOverride;
 			}
 
-			ConnectInfo = FString::Printf(TEXT("%s:%d"), *SessionInfo->HostAddr->ToString(false), HostPort);
+			ConnectInfo = FString::Printf(STEAM_URL_PREFIX TEXT("%s:%d"), *SessionInfo->SteamP2PAddr->ToString(false), SteamPort);
 			bSuccess = true;
 		}
+		else
+		{
+			UE_LOG_ONLINE_SESSION(Warning, TEXT("Steam could not resolve session info! ValidP2P[%d] ValidHost[%d] ConnectionMethod[%s]"), 
+				bP2PDataValid, bHostDataValid, *LexToString(SessionInfo->ConnectionMethod));
+			return false;
+		}
 	}
 
 	return bSuccess;
@@ -1262,7 +1282,7 @@ bool FOnlineSessionSteam::GetResolvedConnectString(FName SessionName, FString& C
 	bool bSuccess = false;
 	// Find the session
 	FNamedOnlineSession* Session = GetNamedSession(SessionName);
-	if (Session != NULL)
+	if (Session != nullptr)
 	{
 		TSharedPtr<FOnlineSessionInfoSteam> SessionInfo = StaticCastSharedPtr<FOnlineSessionInfoSteam>(Session->SessionInfo);
 		if (PortType == NAME_BeaconPort)
@@ -1351,7 +1371,7 @@ FOnlineSessionSettings* FOnlineSessionSteam::GetSessionSettings(FName SessionNam
 	{
 		return &Session->SessionSettings;
 	}
-	return NULL;
+	return nullptr;
 }
 
 void FOnlineSessionSteam::RegisterVoice(const FUniqueNetId& PlayerId)
@@ -1529,7 +1549,7 @@ void FOnlineSessionSteam::Tick(float DeltaTime)
 
 void FOnlineSessionSteam::TickLanTasks(float DeltaTime)
 {
-	if (LANSession != NULL && LANSession->GetBeaconState() > ELanBeaconState::NotUsingLanBeacon)
+	if (LANSession != nullptr && LANSession->GetBeaconState() > ELanBeaconState::NotUsingLanBeacon)
 	{
 		LANSession->Tick(DeltaTime);
 	}
@@ -1541,7 +1561,7 @@ void FOnlineSessionSteam::TickPendingInvites(float DeltaTime)
 	{
 		if (OnSessionUserInviteAcceptedDelegates.IsBound())
 		{
-			FOnlineAsyncItem* NewEvent = NULL;
+			FOnlineAsyncItem* NewEvent = nullptr;
 			FUniqueNetIdSteam FriendId(0);
 			if (PendingInvite.PendingInviteType == ESteamSession::LobbySession)
 			{
@@ -1820,7 +1840,7 @@ void FOnlineSessionSteam::OnLANSearchTimeout()
 		}
 		CurrentSessionSearch->SearchState = EOnlineAsyncTaskState::Done;
 
-		CurrentSessionSearch = NULL;
+		CurrentSessionSearch = nullptr;
 	}
 
 	// Trigger the delegate as complete
@@ -1900,3 +1920,40 @@ void FOnlineSessionSteam::UnregisterLocalPlayer(const FUniqueNetId& PlayerId, FN
 {
 	Delegate.ExecuteIfBound(PlayerId, true);
 }
+
+/** Implementation of the ConnectionMethod converters */
+FString LexToString(const FSteamConnectionMethod Method)
+{
+	switch (Method)
+	{
+	default:
+	case FSteamConnectionMethod::None:
+		return TEXT("None");
+	case FSteamConnectionMethod::Direct:
+		return TEXT("Direct");
+	case FSteamConnectionMethod::P2P:
+		return TEXT("P2P");
+	case FSteamConnectionMethod::PartnerHosted:
+		return TEXT("PartnerHosted");
+	}
+}
+
+FSteamConnectionMethod ToConnectionMethod(const FString& InString)
+{
+	if (InString == TEXT("Direct"))
+	{
+		return FSteamConnectionMethod::Direct;
+	}
+	else if (InString == TEXT("P2P"))
+	{
+		return FSteamConnectionMethod::P2P;
+	}
+	else if (InString == TEXT("PartnerHosted"))
+	{
+		return FSteamConnectionMethod::PartnerHosted;
+	}
+	else
+	{
+		return FSteamConnectionMethod::None;
+	}
+}
\ No newline at end of file
diff --git a/Engine/Plugins/Online/OnlineSubsystemSteam/Source/Private/OnlineSubsystemSteamTypes.h b/Engine/Plugins/Online/OnlineSubsystemSteam/Source/Private/OnlineSubsystemSteamTypes.h
--- a/Engine/Plugins/Online/OnlineSubsystemSteam/Source/Private/OnlineSubsystemSteamTypes.h
+++ b/Engine/Plugins/Online/OnlineSubsystemSteam/Source/Private/OnlineSubsystemSteamTypes.h
@@ -243,6 +243,19 @@ public:
 	}
 };
 
+/** Data regarding preferred session connection methods */
+enum class FSteamConnectionMethod : int8
+{
+	None=0,
+	Direct,
+	P2P,
+	PartnerHosted
+};
+
+/** Functionality for converting a connection method type between formats */
+FString LexToString(const FSteamConnectionMethod Method);
+FSteamConnectionMethod ToConnectionMethod(const FString& InString);
+
 /** 
  * Implementation of session information
  */
@@ -287,6 +300,8 @@ PACKAGE_SCOPE:
 	TSharedPtr<class FInternetAddr> SteamP2PAddr;
 	/** Steam Lobby Id or Gameserver Id if applicable */
 	FUniqueNetIdSteam SessionId;
+	/** How this session should be connected to */
+	FSteamConnectionMethod ConnectionMethod;
 
 public:
 
@@ -302,7 +317,7 @@ public:
 
 	virtual const uint8* GetBytes() const override
 	{
-		return NULL;
+		return nullptr;
 	}
 
 	virtual int32 GetSize() const override
@@ -311,7 +326,8 @@ public:
 			sizeof(ESteamSession::Type) +
 			sizeof(TSharedPtr<class FInternetAddr>) +
 			sizeof(TSharedPtr<class FInternetAddr>) + 
-			sizeof(FUniqueNetIdSteam);
+			sizeof(FUniqueNetIdSteam) +
+			sizeof(FSteamConnectionMethod);
 	}
 
 	virtual bool IsValid() const override
@@ -631,7 +647,7 @@ public:
 			}
 		}
 
-		return NULL;
+		return nullptr;
 	}
 };
 
diff --git a/Engine/Plugins/Online/OnlineSubsystemSteam/Source/Private/SteamSessionKeys.h b/Engine/Plugins/Online/OnlineSubsystemSteam/Source/Private/SteamSessionKeys.h
--- a/Engine/Plugins/Online/OnlineSubsystemSteam/Source/Private/SteamSessionKeys.h
+++ b/Engine/Plugins/Online/OnlineSubsystemSteam/Source/Private/SteamSessionKeys.h
@@ -19,8 +19,8 @@
 #define STEAMKEY_BUILDUNIQUEID "BUILDID"
 
 /** Number of keys above required for valid lobby session */
-#define STEAMKEY_NUMREQUIREDLOBBYKEYS 8
-#define STEAMKEY_NUMREQUIREDSERVERKEYS 3
+#define STEAMKEY_NUMREQUIREDLOBBYKEYS 9
+#define STEAMKEY_NUMREQUIREDSERVERKEYS 4
 
 /** Optional keys (depends on lobby/advertised server session) */
 #define STEAMKEY_HOSTIP "HOSTIP"
@@ -28,6 +28,7 @@
 #define STEAMKEY_P2PADDR "P2PADDR"
 #define STEAMKEY_P2PPORT "P2PPORT"
 #define STEAMKEY_P2PPING "P2PPING"
+#define STEAMKEY_CONNECTIONMETHOD "CONMETHOD"
 
 /**
  * Converts an engine key and its data type to an appropriate Steam key for use with lobbies/gameservers
diff --git a/Engine/Plugins/Runtime/Steam/SteamSockets/Source/SteamSockets/Private/IPAddressSteamSockets.cpp b/Engine/Plugins/Runtime/Steam/SteamSockets/Source/SteamSockets/Private/IPAddressSteamSockets.cpp
--- a/Engine/Plugins/Runtime/Steam/SteamSockets/Source/SteamSockets/Private/IPAddressSteamSockets.cpp
+++ b/Engine/Plugins/Runtime/Steam/SteamSockets/Source/SteamSockets/Private/IPAddressSteamSockets.cpp
@@ -88,7 +88,7 @@ void FInternetAddrSteamSockets::SetIp(const TCHAR* InAddr, bool& bIsValid)
 {
 	FString InAddrStr(InAddr);
 	// See if we're a SteamID.
-	if (InAddrStr.StartsWith(STEAM_URL_PREFIX))
+	if (InAddrStr.StartsWith(STEAM_URL_PREFIX) || InAddrStr.IsNumeric())
 	{
 		InAddrStr.RemoveFromStart(STEAM_URL_PREFIX);
 		FString SteamIPStr, SteamChannelStr;
@@ -179,7 +179,7 @@ FString FInternetAddrSteamSockets::ToString(bool bAppendPort) const
 	{
 		if (GetProtocolType() == FNetworkProtocolTypes::SteamSocketsP2P)
 		{
-			FString BaseResult = FString::Printf(TEXT("%s%llu"), STEAM_URL_PREFIX, Addr.GetSteamID64());
+			FString BaseResult = FString::Printf(TEXT("%llu"), Addr.GetSteamID64());
 			if (bAppendPort)
 			{
 				BaseResult += FString::Printf(TEXT(":%d"), P2PVirtualPort);
