Steam client regression causes calling this function to corrupt the stack when passed with a non-null memory address. Reverting back to an instance where this worked.

#rnx

diff --git a/Engine/Plugins/Online/OnlineSubsystemSteam/Source/Private/OnlineEncryptedAppTicketInterfaceSteam.cpp b/Engine/Plugins/Online/OnlineSubsystemSteam/Source/Private/OnlineEncryptedAppTicketInterfaceSteam.cpp
--- a/Engine/Plugins/Online/OnlineSubsystemSteam/Source/Private/OnlineEncryptedAppTicketInterfaceSteam.cpp
+++ b/Engine/Plugins/Online/OnlineSubsystemSteam/Source/Private/OnlineEncryptedAppTicketInterfaceSteam.cpp
@@ -186,10 +186,7 @@ bool FOnlineEncryptedAppTicketSteam::GetEncryptedAppTicket(TArray<uint8>& OutEnc
 		if (SteamUserPtr)
 		{
 			uint32 ExactTicketSize = 0;
-			// We want to get the ticket size here but the API does not give
-			// this value unless we attempt to write out the ticket.
-			uint32 UnusedByte;
-			SteamUserPtr->GetEncryptedAppTicket((void*)&UnusedByte, 1, &ExactTicketSize);
+			SteamUserPtr->GetEncryptedAppTicket(nullptr, 1, &ExactTicketSize);
 
 			if (ExactTicketSize > 0)
 			{
