Fix nullptr derefrence when cleaning up SteamNetworking connections.

#rnx

diff --git a/Engine/Plugins/Online/OnlineSubsystemSteam/Source/Private/SocketSubsystemSteam.cpp b/Engine/Plugins/Online/OnlineSubsystemSteam/Source/Private/SocketSubsystemSteam.cpp
--- a/Engine/Plugins/Online/OnlineSubsystemSteam/Source/Private/SocketSubsystemSteam.cpp
+++ b/Engine/Plugins/Online/OnlineSubsystemSteam/Source/Private/SocketSubsystemSteam.cpp
@@ -311,12 +311,14 @@ void FSocketSubsystemSteam::ConnectFailure(const FUniqueNetIdSteam& RemoteId)
 	for (int32 ConnIdx=0; ConnIdx<SteamConnections.Num(); ConnIdx++)
 	{
 		USteamNetConnection* SteamConn = CastChecked<USteamNetConnection>(SteamConnections[ConnIdx].Get());
-		TSharedPtr<const FInternetAddrSteam> RemoteAddrSteam = StaticCastSharedPtr<const FInternetAddrSteam>(SteamConn->GetRemoteAddr());
-
-		// Only checking Id here because its a complete failure (channel doesn't matter)
-		if (RemoteAddrSteam->SteamId == RemoteId)
+		if (SteamConn)
 		{
-			SteamConn->Close();
+			TSharedPtr<const FInternetAddrSteam> RemoteAddrSteam = StaticCastSharedPtr<const FInternetAddrSteam>(SteamConn->GetRemoteAddr());
+			// Only checking Id here because its a complete failure (channel doesn't matter)
+			if (RemoteAddrSteam->SteamId == RemoteId)
+			{
+				SteamConn->Close();
+			}
 		}
 	}
 
