Fixing issues with projects linking to SteamAuth (using the reporting delegates). Due to required usage of FUniqueNetIdSteam, code would not build properly outside of the SteamOSS.

This now introduces the FOnlineAuthUtilsSteam class which allows for access to the Steam Auth Interface outside of the SteamOSS code. All delegates have been moved there.

Additionally, new reporting delegates have been added with more information so projects can take better action with SteamAuth results.

No changes necessary for projects just using SteamAuth without any customizations.

diff --git a/Engine/Plugins/Online/OnlineSubsystemSteam/Source/Private/OnlineAuthHandlerSteam.cpp b/Engine/Plugins/Online/OnlineSubsystemSteam/Source/Private/OnlineAuthHandlerSteam.cpp
--- a/Engine/Plugins/Online/OnlineSubsystemSteam/Source/Private/OnlineAuthHandlerSteam.cpp
+++ b/Engine/Plugins/Online/OnlineSubsystemSteam/Source/Private/OnlineAuthHandlerSteam.cpp
@@ -2,6 +2,7 @@
 
 #include "OnlineAuthHandlerSteam.h"
 #include "OnlineAuthInterfaceSteam.h"
+#include "OnlineAuthInterfaceUtilsSteam.h"
 #include "OnlineSubsystemUtils.h"
 
 enum class ESteamAuthMsgType : uint8
@@ -330,7 +331,7 @@ void FSteamAuthHandlerComponent::Incoming(FBitReader& Packet)
 		if (!SteamId.IsValid())
 		{
 			UE_LOG_ONLINE(Error, TEXT("AUTH HANDLER: Got an invalid steamid"));
-			AuthInterface->ExecuteResultDelegate(SteamId, false);
+			AuthInterface->ExecuteResultDelegate(SteamId, false, ESteamAuthResponseCode::NotConnectedToSteam);
 			Packet.SetError();
 			return;
 		}
@@ -339,7 +340,7 @@ void FSteamAuthHandlerComponent::Incoming(FBitReader& Packet)
 		if (!TargetUser.IsValid())
 		{
 			UE_LOG_ONLINE(Error, TEXT("AUTH HANDLER: Could not create user listing for %s"), *SteamId.ToString());
-			AuthInterface->ExecuteResultDelegate(SteamId, false);
+			AuthInterface->ExecuteResultDelegate(SteamId, false, ESteamAuthResponseCode::FailedToCreateUser);
 			Packet.SetError();
 			return;
 		}
diff --git a/Engine/Plugins/Online/OnlineSubsystemSteam/Source/Private/OnlineAuthInterfaceSteam.cpp b/Engine/Plugins/Online/OnlineSubsystemSteam/Source/Private/OnlineAuthInterfaceSteam.cpp
--- a/Engine/Plugins/Online/OnlineSubsystemSteam/Source/Private/OnlineAuthInterfaceSteam.cpp
+++ b/Engine/Plugins/Online/OnlineSubsystemSteam/Source/Private/OnlineAuthInterfaceSteam.cpp
@@ -20,10 +20,11 @@
 // Steam tells us this number in documentation, however there's no define within the SDK
 #define STEAM_AUTH_MAX_TICKET_LENGTH_IN_BYTES 1024
 
-FOnlineAuthSteam::FOnlineAuthSteam(FOnlineSubsystemSteam* InSubsystem) :
+FOnlineAuthSteam::FOnlineAuthSteam(FOnlineSubsystemSteam* InSubsystem, FOnlineAuthSteamUtilsPtr InAuthUtils) :
 	SteamUserPtr(SteamUser()),
 	SteamServerPtr(SteamGameServer()),
 	SteamSubsystem(InSubsystem),
+	AuthUtils(InAuthUtils),
 	bEnabled(false),
 	bBadKey(false),
 	bReuseKey(false),
@@ -65,6 +66,7 @@ FOnlineAuthSteam::FOnlineAuthSteam() :
 	SteamUserPtr(nullptr),
 	SteamServerPtr(nullptr),
 	SteamSubsystem(nullptr),
+	AuthUtils(nullptr),
 	bEnabled(false),
 	bBadKey(false),
 	bReuseKey(false),
@@ -78,8 +80,6 @@ FOnlineAuthSteam::FOnlineAuthSteam() :
 
 FOnlineAuthSteam::~FOnlineAuthSteam()
 {
-	OverrideFailureDelegate.Unbind();
-	OnAuthenticationResultDelegate.Unbind();
 	RevokeAllTickets();
 }
 
@@ -319,9 +319,9 @@ bool FOnlineAuthSteam::KickPlayer(const FUniqueNetId& InUserId, bool bSuppressFa
 	}
 
 	// If we are overridden, respect that.
-	if (OverrideFailureDelegate.IsBound())
+	if (AuthUtils.IsValid() && AuthUtils->OverrideFailureDelegate.IsBound())
 	{
-		OverrideFailureDelegate.Execute(InUserId);
+		AuthUtils->OverrideFailureDelegate.Execute(InUserId);
 		RemoveUser(InUserId);
 		return true;
 	}
@@ -528,14 +528,16 @@ void FOnlineAuthSteam::OnAuthResult(const FUniqueNetId& TargetId, int32 Response
 	{
 		TargetUser->Status |= ESteamAuthStatus::AuthFail;
 	}
-	ExecuteResultDelegate(SteamId, bDidAuthSucceed);
+	ExecuteResultDelegate(SteamId, bDidAuthSucceed, (ESteamAuthResponseCode)Response);
 }
 
-void FOnlineAuthSteam::ExecuteResultDelegate(const FUniqueNetId& TargetId, bool bWasSuccessful)
+void FOnlineAuthSteam::ExecuteResultDelegate(const FUniqueNetId& TargetId, bool bWasSuccessful, ESteamAuthResponseCode ResponseCode)
 {
-	if (OnAuthenticationResultDelegate.IsBound())
+	if (AuthUtils.IsValid())
 	{
-		OnAuthenticationResultDelegate.Execute(TargetId, bWasSuccessful);
+		// Fire both of these delegates.
+		AuthUtils->OnAuthenticationResultDelegate.ExecuteIfBound(TargetId, bWasSuccessful);
+		AuthUtils->OnAuthenticationResultWithCodeDelegate.ExecuteIfBound(TargetId, bWasSuccessful, ResponseCode);
 	}
 }
 
@@ -546,3 +548,11 @@ void FOnlineAuthSteam::FSteamAuthUser::SetKey(const FString& NewKey)
 		RecvTicket = NewKey;
 	}
 }
+
+// Implementation of the public helper to determine if the Auth interface is enabled.
+bool FOnlineAuthUtilsSteam::IsSteamAuthEnabled() const
+{
+	const FOnlineSubsystemSteam* SteamSub = static_cast<const FOnlineSubsystemSteam*>(IOnlineSubsystem::Get(STEAM_SUBSYSTEM));
+	return (SteamSub && SteamSub->GetAuthInterface().IsValid() && SteamSub->GetAuthInterface()->IsSessionAuthEnabled());
+}
+
diff --git a/Engine/Plugins/Online/OnlineSubsystemSteam/Source/Private/OnlineAuthInterfaceSteam.h b/Engine/Plugins/Online/OnlineSubsystemSteam/Source/Private/OnlineAuthInterfaceSteam.h
new file mode 100644
--- /dev/null
+++ b/Engine/Plugins/Online/OnlineSubsystemSteam/Source/Private/OnlineAuthInterfaceSteam.h
@@ -0,0 +1,112 @@
+// Copyright 1998-2019 Epic Games, Inc. All Rights Reserved.
+
+#pragma once
+
+#include "CoreMinimal.h"
+#include "OnlineSubsystemSteam.h"
+#include "OnlineAuthInterfaceUtilsSteam.h"
+#include "OnlineSubsystemSteamTypes.h"
+#include "OnlineSubsystemSteamPackage.h"
+
+/** Steam Authentication Interface.
+ *
+ *  For the most part, this is fully automated. You simply just need to add the packet handler and your server will now
+ *  require Steam Authentication for any incoming users. If a player fails to respond correctly, they will be kicked.
+ *
+ *  For projects that want to interface with SteamAuth and listen to callbacks, check out the OnlineAuthInterfaceUtils header.
+ */
+
+enum class ESteamAuthStatus : uint8
+{
+	None = 0,
+	AuthSuccess = 1 << 0,
+	AuthFail = 1 << 1,
+	ValidationStarted = 1 << 2,
+	KickUser = 1 << 3,
+	FailKick = AuthFail | KickUser,
+	HasOrIsPendingAuth = AuthSuccess | ValidationStarted
+};
+
+ENUM_CLASS_FLAGS(ESteamAuthStatus);
+
+class FOnlineAuthSteam
+{
+PACKAGE_SCOPE:
+	FOnlineAuthSteam(FOnlineSubsystemSteam* InSubsystem, FOnlineAuthSteamUtilsPtr InAuthUtils);
+
+	/** Data pertaining the current authentication state of the users in the game */
+	struct FSteamAuthUser
+	{
+		FSteamAuthUser() : Status(ESteamAuthStatus::None) { }
+		void SetKey(const FString& NewKey);
+		// String representation of another user's ticket. Stored only temporarily
+		FString RecvTicket;
+		ESteamAuthStatus Status;
+	};
+
+	typedef TSharedPtr<FSteamAuthUser, ESPMode::NotThreadSafe> SharedAuthUserSteamPtr;
+
+	SharedAuthUserSteamPtr GetUser(const FUniqueNetId& InUserId);
+	SharedAuthUserSteamPtr GetOrCreateUser(const FUniqueNetId& InUserId);
+
+	bool AuthenticateUser(const FUniqueNetId& InUserId);
+	void EndAuthentication(const FUniqueNetId& InUserId);
+	void MarkPlayerForKick(const FUniqueNetId& InUserId);
+	void RevokeTicket(const uint32& Handle);
+	void RevokeAllTickets();
+	void RemoveUser(const FUniqueNetId& TargetUser);
+
+	/** Generates Steam auth tickets */
+	FString GetAuthTicket(uint32& AuthTokenHandle);
+
+	bool Tick(float DeltaTime);
+	bool Exec(const TCHAR* Cmd);
+
+	/** Callback from Steam messaging */
+	void OnAuthResult(const FUniqueNetId& TargetId, int32 Response);
+
+	void ExecuteResultDelegate(const FUniqueNetId& TargetId, bool bWasSuccessful, ESteamAuthResponseCode ResponseCode);
+
+private:
+	typedef TMap<FUniqueNetIdSteam, SharedAuthUserSteamPtr> SteamAuthentications;
+	SteamAuthentications AuthUsers;
+	TArray<uint32> SteamTicketHandles;
+
+	/** Utility functions */
+	FORCEINLINE bool IsServer() const
+	{
+		return SteamSubsystem != nullptr && SteamSubsystem->IsServer();
+	}
+	bool KickPlayer(const FUniqueNetId& InUserId, bool bSuppressFailure);
+
+	FOnlineAuthSteam();
+
+	/** Steam Interfaces */
+	class ISteamUser* SteamUserPtr;
+	class ISteamGameServer* SteamServerPtr;
+
+	/** Cached pointer to owning subsystem */
+	FOnlineSubsystemSteam* SteamSubsystem;
+	FOnlineAuthSteamUtilsPtr AuthUtils;
+
+	/** Settings */
+	bool bEnabled;
+
+	/** Testing flags */
+PACKAGE_SCOPE:
+	bool bBadKey;		// Send out invalid keys
+	bool bReuseKey;		// Always send out the same key
+	bool bBadWrite;		// Always make the bit writers have errors
+	bool bDropAll;		// Drop all packets
+	bool bRandomDrop;	// Randomly drop packets.
+	bool bNeverSendKey;	// Client never sends their key.
+	bool bSendBadId;	// Always send invalid steam ids.
+
+	/** Setting Getters */
+	bool IsSessionAuthEnabled() const { return bEnabled; }
+
+public:
+	virtual ~FOnlineAuthSteam();
+
+	static uint32 GetMaxTicketSizeInBytes();
+};
diff --git a/Engine/Plugins/Online/OnlineSubsystemSteam/Source/Private/OnlineSubsystemSteam.cpp b/Engine/Plugins/Online/OnlineSubsystemSteam/Source/Private/OnlineSubsystemSteam.cpp
--- a/Engine/Plugins/Online/OnlineSubsystemSteam/Source/Private/OnlineSubsystemSteam.cpp
+++ b/Engine/Plugins/Online/OnlineSubsystemSteam/Source/Private/OnlineSubsystemSteam.cpp
@@ -27,6 +27,7 @@
 #include "OnlineExternalUIInterfaceSteam.h"
 #include "OnlineAchievementsInterfaceSteam.h"
 #include "OnlineAuthInterfaceSteam.h"
+#include "OnlineAuthInterfaceUtilsSteam.h"
 #include "VoiceInterfaceSteam.h"
 
 #include "SteamSharedModule.h"
@@ -225,6 +226,11 @@ FOnlineAuthSteamPtr FOnlineSubsystemSteam::GetAuthInterface() const
 	return AuthInterface;
 }
 
+FOnlineAuthSteamUtilsPtr FOnlineSubsystemSteam::GetAuthInterfaceUtils() const
+{
+	return AuthInterfaceUtils;
+}
+
 IOnlineSessionPtr FOnlineSubsystemSteam::GetSessionInterface() const
 {
 	return SessionInterface;
@@ -442,7 +448,8 @@ bool FOnlineSubsystemSteam::Init()
 
 		PresenceInterface = MakeShareable(new FOnlinePresenceSteam(this));
 		
-		AuthInterface = MakeShareable(new FOnlineAuthSteam(this));
+		AuthInterfaceUtils = MakeShareable(new FOnlineAuthUtilsSteam());
+		AuthInterface = MakeShareable(new FOnlineAuthSteam(this, AuthInterfaceUtils));
 
 		if (!bIsServer)
 		{
@@ -516,6 +523,7 @@ bool FOnlineSubsystemSteam::Shutdown()
 	DESTRUCT_INTERFACE(FriendInterface);
 	DESTRUCT_INTERFACE(IdentityInterface);
 	DESTRUCT_INTERFACE(AuthInterface);
+	DESTRUCT_INTERFACE(AuthInterfaceUtils);
 	DESTRUCT_INTERFACE(SessionInterface);
 	DESTRUCT_INTERFACE(PresenceInterface);
 
diff --git a/Engine/Plugins/Online/OnlineSubsystemSteam/Source/Public/OnlineAuthInterfaceSteam.h b/Engine/Plugins/Online/OnlineSubsystemSteam/Source/Public/OnlineAuthInterfaceSteam.h
deleted file mode 100644
diff --git a/Engine/Plugins/Online/OnlineSubsystemSteam/Source/Public/OnlineAuthInterfaceUtilsSteam.h b/Engine/Plugins/Online/OnlineSubsystemSteam/Source/Public/OnlineAuthInterfaceUtilsSteam.h
new file mode 100644
--- /dev/null
+++ b/Engine/Plugins/Online/OnlineSubsystemSteam/Source/Public/OnlineAuthInterfaceUtilsSteam.h
@@ -0,0 +1,75 @@
+// Copyright 1998-2019 Epic Games, Inc. All Rights Reserved.
+
+#pragma once
+
+#include "CoreMinimal.h"
+
+/** Status flags of SteamAuth response codes, mapped 1:1 with SteamAPI */
+enum class ONLINESUBSYSTEMSTEAM_API ESteamAuthResponseCode : uint8
+{
+	OK = 0,
+	NotConnectedToSteam = 1,
+	NoOwnership = 2,
+	VACBanned = 3,
+	LoggedInElsewhere = 4,
+	VACCheckTimeout = 5,
+	TicketCancelled = 6,
+	TicketAlreadyUsed = 7,
+	TicketInvalid = 8,
+	PublisherBanned = 9,
+	// ... space given for future steam codes ...
+	FailedToCreateUser = 20,
+	Count
+};
+
+/** When authentication has failed and we are about to take action on the user, this delegate is fired. 
+ *	For the auth interface, overriding the delegate exposed in the class allows a game to override the default
+ *	behavior, which is to kick anyone who fails authentication.
+ *	
+ *	If you would like to receive analytics as to the success/failure for users we can identify
+ *	(have their unique net id), use the result delegate instead.
+ */
+DECLARE_MULTICAST_DELEGATE_OneParam(FOnSteamAuthFailure, const class FUniqueNetId& /*FailedUserId*/);
+typedef FOnSteamAuthFailure::FDelegate FOnSteamAuthFailureDelegate;
+
+/** This delegate dictates the success or failure of an authentication result. 
+ *  This means we got a result, but we won't be taking action yet.
+ */
+DECLARE_MULTICAST_DELEGATE_TwoParams(FOnSteamAuthResult, const class FUniqueNetId& /*UserId*/, bool /*bWasSuccess*/);
+typedef FOnSteamAuthResult::FDelegate FOnSteamAuthResultDelegate;
+
+/** This delegate dictates the success or failure of an authentication result with an additional flag that explains the result.
+ *  This is fired at the same time as the above delegate but includes additional data.
+ *  Should a project be listening to result delegate calls, it is recommended to listen to one of these two (but not both).
+ */
+DECLARE_MULTICAST_DELEGATE_ThreeParams(FOnSteamAuthResultWithCode, const class FUniqueNetId& /*UserId*/, bool /*bWasSuccess*/, ESteamAuthResponseCode /*ResponseCode*/);
+typedef FOnSteamAuthResultWithCode::FDelegate FOnSteamAuthResultWithCodeDelegate;
+
+/** A class for interfacing with SteamAuth outside of the SteamOSS. */
+class ONLINESUBSYSTEMSTEAM_API FOnlineAuthUtilsSteam
+{
+public:
+	
+	~FOnlineAuthUtilsSteam()
+	{
+		OverrideFailureDelegate.Unbind();
+		OnAuthenticationResultDelegate.Unbind();
+		OnAuthenticationResultWithCodeDelegate.Unbind();
+	}	
+	
+	/**
+	 * Checks to see if Steam Authentication is enabled due to being defined
+	 * in any PacketHandlerComponent list.
+	 *
+	 * @return the true if SteamAuth is enabled.
+	 *
+	 */
+	bool IsSteamAuthEnabled() const;
+	
+	/** Attach to this delegate to control the behavior of the Steam authentication failure. 
+	 *  This overrides the default behavior (kick). 
+	 */
+	FOnSteamAuthFailureDelegate OverrideFailureDelegate;
+	FOnSteamAuthResultDelegate OnAuthenticationResultDelegate;
+	FOnSteamAuthResultWithCodeDelegate OnAuthenticationResultWithCodeDelegate;
+};
diff --git a/Engine/Plugins/Online/OnlineSubsystemSteam/Source/Public/OnlineSubsystemSteam.h b/Engine/Plugins/Online/OnlineSubsystemSteam/Source/Public/OnlineSubsystemSteam.h
--- a/Engine/Plugins/Online/OnlineSubsystemSteam/Source/Public/OnlineSubsystemSteam.h
+++ b/Engine/Plugins/Online/OnlineSubsystemSteam/Source/Public/OnlineSubsystemSteam.h
@@ -17,6 +17,7 @@ class FOnlineUserCloudSteam;
 class FOnlineVoiceSteam;
 class FOnlinePresenceSteam;
 class FOnlineAuthSteam;
+class FOnlineAuthUtilsSteam;
 
 /** Forward declarations of all interface classes */
 typedef TSharedPtr<class FOnlineSessionSteam, ESPMode::ThreadSafe> FOnlineSessionSteamPtr;
@@ -30,6 +31,7 @@ typedef TSharedPtr<class FOnlineExternalUISteam, ESPMode::ThreadSafe> FOnlineExt
 typedef TSharedPtr<class FOnlineAchievementsSteam, ESPMode::ThreadSafe> FOnlineAchievementsSteamPtr;
 typedef TSharedPtr<class FOnlinePresenceSteam, ESPMode::ThreadSafe> FOnlinePresenceSteamPtr;
 typedef TSharedPtr<class FOnlineAuthSteam, ESPMode::ThreadSafe> FOnlineAuthSteamPtr;
+typedef TSharedPtr<class FOnlineAuthUtilsSteam, ESPMode::ThreadSafe> FOnlineAuthSteamUtilsPtr;
 
 /**
  *	OnlineSubsystemSteam - Implementation of the online subsystem for STEAM services
@@ -95,6 +97,7 @@ protected:
 
 	/** Interface for Steam Session Auth */
 	FOnlineAuthSteamPtr AuthInterface;
+	FOnlineAuthSteamUtilsPtr AuthInterfaceUtils;
 
 	/** Online async task runnable */
 	class FOnlineAsyncTaskManagerSteam* OnlineAsyncTaskThreadRunnable;
@@ -128,6 +131,7 @@ PACKAGE_SCOPE:
 		ExternalUIInterface(nullptr),
 		PresenceInterface(nullptr),
 		AuthInterface(nullptr),
+		AuthInterfaceUtils(nullptr),
 		OnlineAsyncTaskThreadRunnable(nullptr),
 		OnlineAsyncTaskThread(nullptr),
 		SteamAPIClientHandle(nullptr)
@@ -209,6 +213,7 @@ public:
 	}
 
 	virtual FOnlineAuthSteamPtr GetAuthInterface() const;
+	virtual FOnlineAuthSteamUtilsPtr GetAuthInterfaceUtils() const;
 
 	// IOnlineSubsystem
 
