PR #6112: OnlineSubsystemSteam: Fixed FindFriendSession and SendSessionInviteToFriends returning false in some cases (Contributed by [REDACTED_PERSON])

diff --git a/Engine/Plugins/Online/OnlineSubsystemSteam/Source/Private/OnlineSessionInterfaceSteam.cpp b/Engine/Plugins/Online/OnlineSubsystemSteam/Source/Private/OnlineSessionInterfaceSteam.cpp
--- a/Engine/Plugins/Online/OnlineSubsystemSteam/Source/Private/OnlineSessionInterfaceSteam.cpp
+++ b/Engine/Plugins/Online/OnlineSubsystemSteam/Source/Private/OnlineSessionInterfaceSteam.cpp
@@ -1049,6 +1049,7 @@ bool FOnlineSessionSteam::FindFriendSession(int32 LocalUserNum, const FUniqueNet
 
 					FOnlineAsyncTaskSteamFindServerForFriendSession* NewTask = new FOnlineAsyncTaskSteamFindServerForFriendSession(SteamSubsystem, CurrentSessionSearch, LocalUserNum, OnFindFriendSessionCompleteDelegates[LocalUserNum]);
 					SteamSubsystem->QueueAsyncTask(NewTask);
+					bSuccess = true;
 				}
 			}
 		}
@@ -1178,6 +1179,7 @@ bool FOnlineSessionSteam::SendSessionInviteToFriends(int32 LocalUserNum, FName S
 				if (SteamFriends()->InviteUserToGame(FriendId, TCHAR_TO_UTF8(*ConnectionURL)))
 				{
 					UE_LOG_ONLINE_SESSION(Verbose, TEXT("Inviting %s to session %s with %s"), *FriendId.ToDebugString(), *SessionName.ToString(), *ConnectionURL);
+					bSuccess = true;
 				}
 				else
 				{
