Fix byte offset as it was addressed in the steamclient

#rnx

diff --git a/Engine/Plugins/Online/OnlineSubsystemSteam/Source/Private/OnlineEncryptedAppTicketInterfaceSteam.cpp b/Engine/Plugins/Online/OnlineSubsystemSteam/Source/Private/OnlineEncryptedAppTicketInterfaceSteam.cpp
--- a/Engine/Plugins/Online/OnlineSubsystemSteam/Source/Private/OnlineEncryptedAppTicketInterfaceSteam.cpp
+++ b/Engine/Plugins/Online/OnlineSubsystemSteam/Source/Private/OnlineEncryptedAppTicketInterfaceSteam.cpp
@@ -186,7 +186,10 @@ bool FOnlineEncryptedAppTicketSteam::GetEncryptedAppTicket(TArray<uint8>& OutEnc
 		if (SteamUserPtr)
 		{
 			uint32 ExactTicketSize = 0;
-			SteamUserPtr->GetEncryptedAppTicket(nullptr, 1, &ExactTicketSize);
+			// We want to get the ticket size here but the API does not give
+			// this value unless we attempt to write out the ticket.
+			uint32 UnusedByte;
+			SteamUserPtr->GetEncryptedAppTicket((void*)&UnusedByte, 1, &ExactTicketSize);
 
 			if (ExactTicketSize > 0)
 			{
