Deprecating Steam voice packet classes as the Steam platform does not require any platform specific functionality as of UE4.23

diff --git a/Engine/Plugins/Online/OnlineSubsystemSteam/Source/Private/VoicePacketSteam.cpp b/Engine/Plugins/Online/OnlineSubsystemSteam/Source/Private/VoicePacketSteam.cpp
--- a/Engine/Plugins/Online/OnlineSubsystemSteam/Source/Private/VoicePacketSteam.cpp
+++ b/Engine/Plugins/Online/OnlineSubsystemSteam/Source/Private/VoicePacketSteam.cpp
@@ -3,6 +3,7 @@
 #include "VoicePacketSteam.h"
 #include "OnlineSubsystemSteamTypes.h"
 
+PRAGMA_DISABLE_DEPRECATION_WARNINGS
 FVoicePacketSteam::FVoicePacketSteam(const FVoicePacketSteam& Other) :
 	FVoicePacket(Other)
 {
@@ -60,6 +61,4 @@ void FVoicePacketSteam::Serialize(class FArchive& Ar)
 		Ar.Serialize(Buffer.GetData(), Length);
 	}
 }
-
-
-
+PRAGMA_ENABLE_DEPRECATION_WARNINGS
diff --git a/Engine/Plugins/Online/OnlineSubsystemSteam/Source/Private/VoicePacketSteam.h b/Engine/Plugins/Online/OnlineSubsystemSteam/Source/Private/VoicePacketSteam.h
--- a/Engine/Plugins/Online/OnlineSubsystemSteam/Source/Private/VoicePacketSteam.h
+++ b/Engine/Plugins/Online/OnlineSubsystemSteam/Source/Private/VoicePacketSteam.h
@@ -8,7 +8,7 @@
 #include "OnlineSubsystemSteamPackage.h"
 
 /** Defines the data involved in a Steam voice packet */
-class FVoicePacketSteam : public FVoicePacket
+class UE_DEPRECATED(4.25, "Steam Voice is now handled by the default voice implementation, FOnlineVoiceImpl. This code is no longer used.") FVoicePacketSteam : public FVoicePacket
 {
 
 PACKAGE_SCOPE:
@@ -51,10 +51,13 @@ public:
 };
 
 /** Holds the current voice packet data state */
-struct FVoiceDataSteam
+struct UE_DEPRECATED(4.25, "Voice functionality on Steam no longer requires platform specific code.") FVoiceDataSteam
 {
 	/** Data used by the local talkers before sent */
+	PRAGMA_DISABLE_DEPRECATION_WARNINGS
 	FVoicePacketSteam LocalPackets[MAX_SPLITSCREEN_TALKERS];
+	PRAGMA_ENABLE_DEPRECATION_WARNINGS
+
 	/** Holds the set of received packets that need to be processed */
 	FVoicePacketList RemotePackets;
 
