Fixing frivolous deprecation warnings generating with MSVC when building Steam.

diff --git a/Engine/Plugins/Online/OnlineSubsystemSteam/Source/Private/VoicePacketSteam.h b/Engine/Plugins/Online/OnlineSubsystemSteam/Source/Private/VoicePacketSteam.h
--- a/Engine/Plugins/Online/OnlineSubsystemSteam/Source/Private/VoicePacketSteam.h
+++ b/Engine/Plugins/Online/OnlineSubsystemSteam/Source/Private/VoicePacketSteam.h
@@ -50,14 +50,18 @@ public:
 	//~ End FVoicePacket interface
 };
 
-/** Holds the current voice packet data state */
-struct UE_DEPRECATED(4.25, "Voice functionality on Steam no longer requires platform specific code.") FVoiceDataSteam
+PRAGMA_DISABLE_DEPRECATION_WARNINGS
+struct FVoiceDataSteam_DEPRECATED
 {
 	/** Data used by the local talkers before sent */
-	PRAGMA_DISABLE_DEPRECATION_WARNINGS
 	FVoicePacketSteam LocalPackets[MAX_SPLITSCREEN_TALKERS];
-	PRAGMA_ENABLE_DEPRECATION_WARNINGS
+};
+PRAGMA_ENABLE_DEPRECATION_WARNINGS
+
 
+/** Holds the current voice packet data state */
+struct UE_DEPRECATED(4.25, "Voice functionality on Steam no longer requires platform specific code.") FVoiceDataSteam : public FVoiceDataSteam_DEPRECATED
+{
 	/** Holds the set of received packets that need to be processed */
 	FVoicePacketList RemotePackets;
 
