Fixing Steam OSS initializing a server when starting a client build. This is not an outdated system that is no longer necessary. Old initialization process can still be used by setting OnlineSubsystemSteam.bInitServerOnClient to true in Engine inis.

diff --git a/Engine/Config/BaseEngine.ini b/Engine/Config/BaseEngine.ini
--- a/Engine/Config/BaseEngine.ini
+++ b/Engine/Config/BaseEngine.ini
@@ -1404,6 +1404,7 @@ DefaultUntrustedDisplayName=Player
 bEnabled=false
 SteamDevAppId=0
 GameServerQueryPort=27015
+bInitServerOnClient=false
 bRelaunchInSteam=false
 GameVersion=1.0.0.0
 bVACEnabled=1
diff --git a/Engine/Plugins/Online/OnlineSubsystemSteam/Source/Private/OnlineSubsystemSteam.cpp b/Engine/Plugins/Online/OnlineSubsystemSteam/Source/Private/OnlineSubsystemSteam.cpp
--- a/Engine/Plugins/Online/OnlineSubsystemSteam/Source/Private/OnlineSubsystemSteam.cpp
+++ b/Engine/Plugins/Online/OnlineSubsystemSteam/Source/Private/OnlineSubsystemSteam.cpp
@@ -48,24 +48,23 @@ namespace FNetworkProtocolTypes
 	const FLazyName Steam(TEXT("Steam"));
 }
 
-#if !UE_BUILD_SHIPPING
 namespace OSSConsoleVariables
 {
-	/** This CVar is for use by NetcodeUnitTest, and is completely unsupported for use outside of this */
 	TAutoConsoleVariable<int32> CVarSteamInitServerOnClient(
 			TEXT("OSS.SteamInitServerOnClient"),
-			1,
-			TEXT("Whether or not to initialize the Steam server interface on clients (default true)"),
-			ECVF_Default);
+			0,
+			TEXT("Whether or not to initialize the Steam server interface on clients (default false)"),
+			ECVF_Default | ECVF_Cheat);
 
+#if !UE_BUILD_SHIPPING
 	/** CVar used by NetcodeUnitTest, to force-enable Steam within the unit test commandlet */
 	TAutoConsoleVariable<int32> CVarSteamUnitTest(
 			TEXT("OSS.SteamUnitTest"),
 			0,
 			TEXT("Whether or not Steam is being force-enabled by NetcodeUnitTest"),
 			ECVF_Default);
-}
 #endif
+}
 
 
 extern "C" 
@@ -431,20 +430,17 @@ bool FOnlineSubsystemSteam::Init()
 	}
 
 	const bool bIsServer = IsRunningDedicatedServer();
-	bool bAttemptServerInit = true;
+	bool bInitServerOnClient = false;
+	GConfig->GetBool(TEXT("OnlineSubsystemSteam"), TEXT("bInitServerOnClient"), bInitServerOnClient, GEngineIni);
+	bool bAttemptServerInit = bIsServer || !!OSSConsoleVariables::CVarSteamInitServerOnClient.GetValueOnGameThread() || bInitServerOnClient;
 
-#if !UE_BUILD_SHIPPING
-	// Add a bypass for NetcodeUnitTest, to allow running a Steam server + client from same machine, by disabling server init on client.
-	// This is an unapproved/unsupported method for using OnlineSubsystemSteam.
-	bAttemptServerInit = bIsServer || !!OSSConsoleVariables::CVarSteamInitServerOnClient.GetValueOnGameThread();
-#endif
+	UE_LOG_ONLINE(Verbose, TEXT("Steam: Starting SteamWorks. Client [%d] Server [%d]"), !bIsServer, bAttemptServerInit);
 	
 	// Don't initialize the Steam Client API if we are launching as a server
 	bool bClientInitSuccess = !bIsServer ? InitSteamworksClient(bRelaunchInSteam, RelaunchAppId) : true;
 
-	// Initialize the Steam Server API if this is a dedicated server or
-	//  the Client API was successfully initialized
-	bool bServerInitSuccess = bClientInitSuccess ? (!bAttemptServerInit || InitSteamworksServer()) : false;
+	// Initialize the Steam Server API if this is a dedicated server or servers should initialize on clients
+	bool bServerInitSuccess = bAttemptServerInit ? (InitSteamworksServer()) : true;
 
 	if (bClientInitSuccess && bServerInitSuccess)
 	{
