Fix for dubious warning during MSVC static analysis
https://developercommunity2.visualstudio.com/t/C6385-False-Positive/878703

diff --git a/Engine/Plugins/Online/OnlineSubsystemSteam/Source/Private/OnlineSessionAsyncServerSteam.cpp b/Engine/Plugins/Online/OnlineSubsystemSteam/Source/Private/OnlineSessionAsyncServerSteam.cpp
--- a/Engine/Plugins/Online/OnlineSubsystemSteam/Source/Private/OnlineSessionAsyncServerSteam.cpp
+++ b/Engine/Plugins/Online/OnlineSubsystemSteam/Source/Private/OnlineSessionAsyncServerSteam.cpp
@@ -813,6 +813,12 @@ void FPendingSearchResultSteam::CancelQuery()
 	}
 }
 
+
+#ifdef _MSC_VER
+#pragma warning(push)
+#pragma warning(disable : 6385) // Disable dubious Static Analysis warning C6385 on windows. See MS TechNote Here https://developercommunity2.visualstudio.com/t/C6385-False-Positive/878703
+#endif
+
 /**
  *  Create the proper query for the master server based on the given search settings
  *
@@ -983,6 +989,10 @@ void FOnlineAsyncTaskSteamFindServerBase::CreateQuery(MatchMakingKeyValuePair_t*
 	}
 }
 
+#ifdef _MSC_VER
+#pragma warning(pop)
+#endif
+
 /**
  *	Create a search result from a server response
  *
