Cleaning up UE4 references in OnlineSubsystem code

diff --git a/Engine/Plugins/Online/OnlineSubsystemSteam/Source/Private/OnlineSessionAsyncServerSteam.cpp b/Engine/Plugins/Online/OnlineSubsystemSteam/Source/Private/OnlineSessionAsyncServerSteam.cpp
--- a/Engine/Plugins/Online/OnlineSubsystemSteam/Source/Private/OnlineSessionAsyncServerSteam.cpp
+++ b/Engine/Plugins/Online/OnlineSubsystemSteam/Source/Private/OnlineSessionAsyncServerSteam.cpp
@@ -24,32 +24,63 @@
  * If these do not match the same values on the Steam partner backend, 
  * matchmaking will not work for your dedicated server
  */
-#ifndef UE4_PROJECT_STEAMPRODUCTNAME
+#ifndef UE_PROJECT_STEAMPRODUCTNAME
+
 #ifdef STEAMPRODUCTNAME
-UE_DEPRECATED(4.22, "The Steam Product Name Macro has been updated to be configurable from your Target.cs file. Please change STEAMPRODUCTNAME to UE4_PROJECT_STEAMPRODUCTNAME.")
-#define UE4_PROJECT_STEAMPRODUCTNAME STEAMPRODUCTNAME
+UE_DEPRECATED(4.22, "The Steam Product Name Macro has been updated to be configurable from your Target.cs file. Please change STEAMPRODUCTNAME to UE_PROJECT_STEAMPRODUCTNAME.")
+#define UE_PROJECT_STEAMPRODUCTNAME STEAMPRODUCTNAME
 #else
-#define UE4_PROJECT_STEAMPRODUCTNAME "unrealdk"
-#endif
-#endif
 
-#ifndef UE4_PROJECT_STEAMGAMEDIR
+#ifdef UE4_PROJECT_STEAMPRODUCTNAME
+UE_DEPRECATED(5.0, "UE4_PROJECT_STEAMPRODUCTNAME has been renamed to UE_PROJECT_STEAMPRODUCTNAME.")
+#define UE_PROJECT_STEAMPRODUCTNAME UE4_PROJECT_STEAMPRODUCTNAME
+#else
+#define UE_PROJECT_STEAMPRODUCTNAME "unrealdk"
+#endif //UE4_PROJECT_STEAMPRODUCTNAME
+
+#endif //STEAMPRODUCTNAME
+
+#endif //UE_PROJECT_STEAMPRODUCTNAME
+
+
+
+#ifndef UE_PROJECT_STEAMGAMEDIR
+
 #ifdef STEAMGAMEDIR
-UE_DEPRECATED(4.22, "The Steam Game Directory Macro has been updated to be configurable from your Target.cs file. Please change STEAMGAMEDIR to UE4_PROJECT_STEAMGAMEDIR.")
-#define UE4_PROJECT_STEAMGAMEDIR STEAMGAMEDIR
+UE_DEPRECATED(4.22, "The Steam Game Directory Macro has been updated to be configurable from your Target.cs file. Please change STEAMGAMEDIR to UE_PROJECT_STEAMGAMEDIR.")
+#define UE_PROJECT_STEAMGAMEDIR STEAMGAMEDIR
 #else
-#define UE4_PROJECT_STEAMGAMEDIR "unrealtest"
-#endif
-#endif
 
-#ifndef UE4_PROJECT_STEAMGAMEDESC
+#ifdef UE4_PROJECT_STEAMGAMEDIR
+UE_DEPRECATED(5.0, "UE4_PROJECT_STEAMGAMEDIR has been renamed to UE_PROJECT_STEAMGAMEDIR.")
+#define UE_PROJECT_STEAMGAMEDIR UE4_PROJECT_STEAMGAMEDIR
+#else
+#define UE_PROJECT_STEAMGAMEDIR "unrealtest"
+#endif // UE4_PROJECT_STEAMPRODUCTNAME
+
+#endif // STEAMGAMEDIR
+
+#endif // UE_PROJECT_STEAMGAMEDIR
+
+
+
+#ifndef UE_PROJECT_STEAMGAMEDESC
+
 #ifdef STEAMGAMEDESC
-UE_DEPRECATED(4.22, "The Steam Game Description Macro has been updated to be configurable from your Target.cs file. Please change STEAMGAMEDESC to UE4_PROJECT_STEAMGAMEDESC.")
-#define UE4_PROJECT_STEAMGAMEDESC STEAMGAMEDESC
+UE_DEPRECATED(4.22, "The Steam Game Description Macro has been updated to be configurable from your Target.cs file. Please change STEAMGAMEDESC to UE_PROJECT_STEAMGAMEDESC.")
+#define UE_PROJECT_STEAMGAMEDESC STEAMGAMEDESC
 #else
-#define UE4_PROJECT_STEAMGAMEDESC "Unreal Test!"
-#endif
-#endif
+
+#ifdef UE4_PROJECT_STEAMGAMEDESC
+UE_DEPRECATED(5.0, "UE4_PROJECT_STEAMGAMEDESC has been renamed to UE_PROJECT_STEAMGAMEDESC.")
+#define UE_PROJECT_STEAMGAMEDESC UE4_PROJECT_STEAMGAMEDESC
+#else
+#define UE_PROJECT_STEAMGAMEDESC "Unreal Test!"
+#endif // UE4_PROJECT_STEAMGAMEDESC
+
+#endif // STEAMGAMEDESC
+
+#endif // UE_PROJECT_STEAMGAMEDESC
 
 /**
  * Get the engine unique build id as Steam key
@@ -377,13 +408,13 @@ void FOnlineAsyncTaskSteamCreateServer::Tick()
 			bool bWantsDedicated = Session->SessionSettings.bIsDedicated;
 			UE_LOG_ONLINE_SESSION(Verbose, TEXT("Starting Steam game server. Dedicated? %d Game Dir is: %s Product Name is: %s\nGame Desc is: %s"),
 				bWantsDedicated,
-				ANSI_TO_TCHAR((UE4_PROJECT_STEAMGAMEDIR)),
-				ANSI_TO_TCHAR((UE4_PROJECT_STEAMPRODUCTNAME)),
-				ANSI_TO_TCHAR((UE4_PROJECT_STEAMGAMEDESC)));
+				ANSI_TO_TCHAR((UE_PROJECT_STEAMGAMEDIR)),
+				ANSI_TO_TCHAR((UE_PROJECT_STEAMPRODUCTNAME)),
+				ANSI_TO_TCHAR((UE_PROJECT_STEAMGAMEDESC)));
 
-			SteamGameServerPtr->SetModDir(UE4_PROJECT_STEAMGAMEDIR);
-			SteamGameServerPtr->SetProduct(UE4_PROJECT_STEAMPRODUCTNAME);
-			SteamGameServerPtr->SetGameDescription(UE4_PROJECT_STEAMGAMEDESC);
+			SteamGameServerPtr->SetModDir(UE_PROJECT_STEAMGAMEDIR);
+			SteamGameServerPtr->SetProduct(UE_PROJECT_STEAMPRODUCTNAME);
+			SteamGameServerPtr->SetGameDescription(UE_PROJECT_STEAMGAMEDESC);
 			SteamGameServerPtr->SetDedicatedServer(bWantsDedicated);
 
 			if (!SteamGameServerPtr->BLoggedOn())
@@ -842,7 +873,7 @@ void FOnlineAsyncTaskSteamFindServerBase::CreateQuery(MatchMakingKeyValuePair_t*
     NumFilters = 0;
 	// Filter must match at least our game
 	FCStringAnsi::Strncpy(Filters[NumFilters].m_szKey, "gamedir", KeySize);
-	FCStringAnsi::Strncpy(Filters[NumFilters].m_szValue, UE4_PROJECT_STEAMGAMEDIR, ValueSize);
+	FCStringAnsi::Strncpy(Filters[NumFilters].m_szValue, UE_PROJECT_STEAMGAMEDIR, ValueSize);
 	NumFilters++;
 
 	FString MapName;
diff --git a/Engine/Plugins/Online/OnlineSubsystemSteam/Source/Private/OnlineSubsystemSteam.cpp b/Engine/Plugins/Online/OnlineSubsystemSteam/Source/Private/OnlineSubsystemSteam.cpp
--- a/Engine/Plugins/Online/OnlineSubsystemSteam/Source/Private/OnlineSubsystemSteam.cpp
+++ b/Engine/Plugins/Online/OnlineSubsystemSteam/Source/Private/OnlineSubsystemSteam.cpp
@@ -40,8 +40,13 @@
  * This helps the SteamAPI find your project on shipping builds
  * if your game is launched outside of Steam.
  */
-#ifndef UE4_PROJECT_STEAMSHIPPINGID
-#define UE4_PROJECT_STEAMSHIPPINGID 0
+#ifndef UE_PROJECT_STEAMSHIPPINGID
+#ifdef UE4_PROJECT_STEAMSHIPPINGID
+UE_DEPRECATED(5.0, "UE4_PROJECT_STEAMSHIPPINGID has been renamed to UE_PROJECT_STEAMSHIPPINGID.")
+#define UE_PROJECT_STEAMSHIPPINGID UE4_PROJECT_STEAMSHIPPINGID
+#else
+#define UE_PROJECT_STEAMSHIPPINGID 0
+#endif
 #endif
 
 namespace FNetworkProtocolTypes
@@ -217,7 +222,7 @@ bool ConfigureSteamInitDevOptions(bool& RequireRelaunch, int32& RelaunchAppId)
 #else
 	// Always check against the Steam client when shipping
 	RequireRelaunch = true;
-	RelaunchAppId = UE4_PROJECT_STEAMSHIPPINGID;
+	RelaunchAppId = UE_PROJECT_STEAMSHIPPINGID;
 #endif
 
 	return true;
diff --git a/Engine/Plugins/Online/OnlineSubsystemUtils/Source/OnlineSubsystemUtils/Private/Tests/TestPurchaseInterface.h b/Engine/Plugins/Online/OnlineSubsystemUtils/Source/OnlineSubsystemUtils/Private/Tests/TestPurchaseInterface.h
--- a/Engine/Plugins/Online/OnlineSubsystemUtils/Source/OnlineSubsystemUtils/Private/Tests/TestPurchaseInterface.h
+++ b/Engine/Plugins/Online/OnlineSubsystemUtils/Source/OnlineSubsystemUtils/Private/Tests/TestPurchaseInterface.h
@@ -2,7 +2,7 @@
 
 #pragma once
 
-// UE4 includes
+// UE includes
 
 #include "CoreMinimal.h"
 #include "Online/CoreOnline.h"
diff --git a/Engine/Plugins/Online/OnlineSubsystemUtils/Source/OnlineSubsystemUtils/Private/Tests/TestSharingInterface.cpp b/Engine/Plugins/Online/OnlineSubsystemUtils/Source/OnlineSubsystemUtils/Private/Tests/TestSharingInterface.cpp
--- a/Engine/Plugins/Online/OnlineSubsystemUtils/Source/OnlineSubsystemUtils/Private/Tests/TestSharingInterface.cpp
+++ b/Engine/Plugins/Online/OnlineSubsystemUtils/Source/OnlineSubsystemUtils/Private/Tests/TestSharingInterface.cpp
@@ -36,7 +36,7 @@ void FTestSharingInterface::Test(UWorld* InWorld, bool bWithImage)
 	SharingInterface = OnlineSub->GetSharingInterface();
 	if (SharingInterface.IsValid())
 	{
-		TestStatusUpdate.Message = FString::Printf(TEXT("This is a test post for UE4 Sharing support! Date = %s"), *FDateTime::Now().ToString());
+		TestStatusUpdate.Message = FString::Printf(TEXT("This is a test post for UE Sharing support! Date = %s"), *FDateTime::Now().ToString());
 		TestStatusUpdate.PostPrivacy = EOnlineStatusUpdatePrivacy::OnlyMe;
 		if (bWithImage)
 		{
diff --git a/Engine/Plugins/Online/OnlineSubsystemUtils/Source/OnlineSubsystemUtils/Private/Tests/TestSharingInterface.h b/Engine/Plugins/Online/OnlineSubsystemUtils/Source/OnlineSubsystemUtils/Private/Tests/TestSharingInterface.h
--- a/Engine/Plugins/Online/OnlineSubsystemUtils/Source/OnlineSubsystemUtils/Private/Tests/TestSharingInterface.h
+++ b/Engine/Plugins/Online/OnlineSubsystemUtils/Source/OnlineSubsystemUtils/Private/Tests/TestSharingInterface.h
@@ -2,7 +2,7 @@
 
 #pragma once
 
-// UE4 includes
+// UE includes
 
 #include "CoreMinimal.h"
 #include "OnlineSubsystem.h"
diff --git a/Engine/Plugins/Online/OnlineSubsystemUtils/Source/OnlineSubsystemUtils/Private/Tests/TestStoreInterface.h b/Engine/Plugins/Online/OnlineSubsystemUtils/Source/OnlineSubsystemUtils/Private/Tests/TestStoreInterface.h
--- a/Engine/Plugins/Online/OnlineSubsystemUtils/Source/OnlineSubsystemUtils/Private/Tests/TestStoreInterface.h
+++ b/Engine/Plugins/Online/OnlineSubsystemUtils/Source/OnlineSubsystemUtils/Private/Tests/TestStoreInterface.h
@@ -2,7 +2,7 @@
 
 #pragma once
 
-// UE4 includes
+// UE includes
 
 #include "CoreMinimal.h"
 #include "Online/CoreOnline.h"
diff --git a/Engine/Plugins/Online/OnlineSubsystemUtils/Source/OnlineSubsystemUtils/Private/Tests/TestUserInterface.h b/Engine/Plugins/Online/OnlineSubsystemUtils/Source/OnlineSubsystemUtils/Private/Tests/TestUserInterface.h
--- a/Engine/Plugins/Online/OnlineSubsystemUtils/Source/OnlineSubsystemUtils/Private/Tests/TestUserInterface.h
+++ b/Engine/Plugins/Online/OnlineSubsystemUtils/Source/OnlineSubsystemUtils/Private/Tests/TestUserInterface.h
@@ -2,7 +2,7 @@
 
 #pragma once
 
-// UE4 includes
+// UE includes
 
 #include "CoreMinimal.h"
 #include "Online/CoreOnline.h"
