- Change to using session auth token in OSS Steam

Approved by BenM

diff --git a/Engine/Plugins/Online/OnlineSubsystemEOS/Source/OnlineSubsystemEOS/Private/UserManagerEOS.cpp b/Engine/Plugins/Online/OnlineSubsystemEOS/Source/OnlineSubsystemEOS/Private/UserManagerEOS.cpp
--- a/Engine/Plugins/Online/OnlineSubsystemEOS/Source/OnlineSubsystemEOS/Private/UserManagerEOS.cpp
+++ b/Engine/Plugins/Online/OnlineSubsystemEOS/Source/OnlineSubsystemEOS/Private/UserManagerEOS.cpp
@@ -95,7 +95,7 @@ static inline EOS_EExternalCredentialType ToEOS_EExternalCredentialType(FName OS
 {
 	if (OSSName == STEAM_SUBSYSTEM)
 	{
-		return EOS_EExternalCredentialType::EOS_ECT_STEAM_APP_TICKET;
+		return EOS_EExternalCredentialType::EOS_ECT_STEAM_SESSION_TICKET;
 	}
 	else if (OSSName == PS4_SUBSYSTEM || USE_PSN_ID_TOKEN)
 	{
diff --git a/Engine/Plugins/Online/OnlineSubsystemSteam/Source/Private/OnlineIdentityInterfaceSteam.cpp b/Engine/Plugins/Online/OnlineSubsystemSteam/Source/Private/OnlineIdentityInterfaceSteam.cpp
--- a/Engine/Plugins/Online/OnlineSubsystemSteam/Source/Private/OnlineIdentityInterfaceSteam.cpp
+++ b/Engine/Plugins/Online/OnlineSubsystemSteam/Source/Private/OnlineIdentityInterfaceSteam.cpp
@@ -241,15 +241,7 @@ FString FOnlineIdentitySteam::GetAuthType() const
 
 void FOnlineIdentitySteam::GetLinkedAccountAuthToken(int32 LocalUserNum, const FOnGetLinkedAccountAuthTokenCompleteDelegate& Delegate) const
 {
-	SteamSubsystem->GetEncryptedAppTicketInterface()->OnEncryptedAppTicketResultDelegate.AddLambda([this, LocalUserNum, OnComplete = FOnGetLinkedAccountAuthTokenCompleteDelegate(Delegate)](bool bEncryptedDataAvailable, int32 ResultCode)
-	{
-		FExternalAuthToken ExternalToken;
-		if (bEncryptedDataAvailable)
-		{
-			SteamSubsystem->GetEncryptedAppTicketInterface()->GetEncryptedAppTicket(ExternalToken.TokenData);
-		}
-		// Pass the info back to the original caller
-		OnComplete.ExecuteIfBound(LocalUserNum, ExternalToken.HasTokenData(), ExternalToken);
-	});
-	SteamSubsystem->GetEncryptedAppTicketInterface()->RequestEncryptedAppTicket(nullptr, 0);
+	FExternalAuthToken AuthToken;
+	AuthToken.TokenString = GetAuthToken(LocalUserNum);
+	Delegate.ExecuteIfBound(LocalUserNum, AuthToken.HasTokenString(), AuthToken);
 }
