- Add TokenType parameter to GetLinkedAccountAuthToken
- Steam: Revert previous breaking change to GetLinkedAccountAuthToken - It originally returned APP tokens, we changed to SESSION tokens, and this broke auth flows for some licensees. Now reverted to returning APP tokens by default.
- Steam: You can now ask GetLinkedAccountAuthToken explicitly for APP or SESSION tokens. If an empty TokenType is passed, what it returns is configurable, but by default is APP.

diff --git a/Engine/Plugins/Online/OnlineServicesOSSAdapter/Source/Private/Online/AuthOSSAdapter.cpp b/Engine/Plugins/Online/OnlineServicesOSSAdapter/Source/Private/Online/AuthOSSAdapter.cpp
--- a/Engine/Plugins/Online/OnlineServicesOSSAdapter/Source/Private/Online/AuthOSSAdapter.cpp
+++ b/Engine/Plugins/Online/OnlineServicesOSSAdapter/Source/Private/Online/AuthOSSAdapter.cpp
@@ -445,7 +445,7 @@ TOnlineAsyncOpHandle<FAuthQueryExternalAuthToken> FAuthOSSAdapter::QueryExternal
 		}
 		else if (EnumHasAnyFlags(ExternalAuthTokenTranslationTraits->Flags, EExternalAuthTokenTranslationFlags::TokenString | EExternalAuthTokenTranslationFlags::TokenBinary))
 		{
-			GetIdentityInterface()->GetLinkedAccountAuthToken(AccountInfoOSSAdapter->LocalUserNum, *MakeDelegateAdapter(this,
+			GetIdentityInterface()->GetLinkedAccountAuthToken(AccountInfoOSSAdapter->LocalUserNum, FString(), *MakeDelegateAdapter(this,
 			[this, WeakOp = InAsyncOp.AsWeak(), Promise = MoveTemp(Promise), ExternalAuthTokenTranslationTraits](int32 LocalUserNum, bool bWasSuccessful, const ::FExternalAuthToken& AuthToken) mutable
 			{
 				if (TSharedPtr<TOnlineAsyncOp<FAuthQueryExternalAuthToken>> Op = WeakOp.Pin())
diff --git a/Engine/Plugins/Online/OnlineSubsystem/Source/Public/Interfaces/OnlineIdentityInterface.h b/Engine/Plugins/Online/OnlineSubsystem/Source/Public/Interfaces/OnlineIdentityInterface.h
--- a/Engine/Plugins/Online/OnlineSubsystem/Source/Public/Interfaces/OnlineIdentityInterface.h
+++ b/Engine/Plugins/Online/OnlineSubsystem/Source/Public/Interfaces/OnlineIdentityInterface.h
@@ -459,13 +459,21 @@ public:
 	 * Gets a user's platform specific authentication ticket used to sign into linked account(s)
 	 *
 	 * @param LocalUserNum the controller number of the associated user
+	 * @param TokenType the platform specific token type you are requesting.
+	 * @param Delegate delegate to execute when the async task completes
 	 */
-	virtual void GetLinkedAccountAuthToken(int32 LocalUserNum, const FOnGetLinkedAccountAuthTokenCompleteDelegate& Delegate) const
+	virtual void GetLinkedAccountAuthToken(int32 LocalUserNum, const FString& TokenType, const FOnGetLinkedAccountAuthTokenCompleteDelegate& Delegate) const
 	{
 		FExternalAuthToken EmptyToken;
 		Delegate.ExecuteIfBound(LocalUserNum, false, EmptyToken);
 	}
 
+	UE_DEPRECATED(5.2, "Please use GetLinkedAccountAuthToken taking a TokenType")
+	virtual void GetLinkedAccountAuthToken(int32 LocalUserNum, const FOnGetLinkedAccountAuthTokenCompleteDelegate& Delegate) const
+	{
+		GetLinkedAccountAuthToken(LocalUserNum, FString(), Delegate);
+	}
+
 	/**
 	 * Delegate executed when we get a user privilege result.
 	 *
diff --git a/Engine/Plugins/Online/OnlineSubsystemEOS/Source/OnlineSubsystemEOS/Private/UserManagerEOS.cpp b/Engine/Plugins/Online/OnlineSubsystemEOS/Source/OnlineSubsystemEOS/Private/UserManagerEOS.cpp
--- a/Engine/Plugins/Online/OnlineSubsystemEOS/Source/OnlineSubsystemEOS/Private/UserManagerEOS.cpp
+++ b/Engine/Plugins/Online/OnlineSubsystemEOS/Source/OnlineSubsystemEOS/Private/UserManagerEOS.cpp
@@ -1,17 +1,20 @@
 // Copyright Epic Games, Inc. All Rights Reserved.
 
 #include "UserManagerEOS.h"
-#include "OnlineSubsystemEOS.h"
-#include "OnlineSubsystemEOSPrivate.h"
+
 #include "Misc/CommandLine.h"
 #include "Misc/Guid.h"
 #include "Misc/OutputDeviceRedirector.h"
-#include "IPAddress.h"
-#include "SocketSubsystem.h"
-#include "OnlineError.h"
-#include "EOSSettings.h"
+
 #include "CoreMinimal.h"
+#include "EOSSettings.h"
 #include "IEOSSDKManager.h"
+#include "IPAddress.h"
+#include "OnlineError.h"
+#include "OnlineSubsystemEOS.h"
+#include "OnlineSubsystemEOSPrivate.h"
+#include "OnlineSubsystemNames.h"
+#include "SocketSubsystem.h"
 
 #include COMPILED_PLATFORM_HEADER(EOSHelpers.h)
 
@@ -274,8 +277,16 @@ void FUserManagerEOS::GetPlatformAuthToken(int32 LocalUserNum, const FOnGetLinke
 		Delegate.ExecuteIfBound(LocalUserNum, false, FExternalAuthToken());
 		return;
 	}
+
+	FString TokenType;
+	// TODO config map of OSS -> token type?
+	if (PlatformOSS->GetSubsystemName() == STEAM_SUBSYSTEM)
+	{
+		TokenType = TEXT("Session");
+	}
+
 	// Request the auth token from the platform
-	PlatformIdentity->GetLinkedAccountAuthToken(LocalUserNum, Delegate);
+	PlatformIdentity->GetLinkedAccountAuthToken(LocalUserNum, TokenType, Delegate);
 }
 
 FString FUserManagerEOS::GetPlatformDisplayName(int32 LocalUserNum) const
@@ -1547,7 +1558,7 @@ FPlatformUserId FUserManagerEOS::GetPlatformUserIdFromUniqueNetId(const FUniqueN
 	return GetPlatformUserIdFromLocalUserNum(GetLocalUserNumFromUniqueNetId(UniqueNetId));
 }
 
-void FUserManagerEOS::GetLinkedAccountAuthToken(int32 LocalUserNum, const FOnGetLinkedAccountAuthTokenCompleteDelegate& Delegate) const
+void FUserManagerEOS::GetLinkedAccountAuthToken(int32 LocalUserNum, const FString& /*TokenType*/, const FOnGetLinkedAccountAuthTokenCompleteDelegate& Delegate) const
 {
 	FExternalAuthToken ExternalToken;
 	ExternalToken.TokenString = GetAuthToken(LocalUserNum);
diff --git a/Engine/Plugins/Online/OnlineSubsystemEOS/Source/OnlineSubsystemEOS/Private/UserManagerEOS.h b/Engine/Plugins/Online/OnlineSubsystemEOS/Source/OnlineSubsystemEOS/Private/UserManagerEOS.h
--- a/Engine/Plugins/Online/OnlineSubsystemEOS/Source/OnlineSubsystemEOS/Private/UserManagerEOS.h
+++ b/Engine/Plugins/Online/OnlineSubsystemEOS/Source/OnlineSubsystemEOS/Private/UserManagerEOS.h
@@ -219,7 +219,7 @@ public:
 	virtual FString GetAuthType() const override;
 	virtual void RevokeAuthToken(const FUniqueNetId& LocalUserId, const FOnRevokeAuthTokenCompleteDelegate& Delegate) override;
 	virtual FPlatformUserId GetPlatformUserIdFromUniqueNetId(const FUniqueNetId& UniqueNetId) const override;
-	virtual void GetLinkedAccountAuthToken(int32 LocalUserNum, const FOnGetLinkedAccountAuthTokenCompleteDelegate& Delegate) const override;
+	virtual void GetLinkedAccountAuthToken(int32 LocalUserNum, const FString& TokenType, const FOnGetLinkedAccountAuthTokenCompleteDelegate& Delegate) const override;
 // ~IOnlineIdentity Interface
 	ELoginStatus::Type GetLoginStatus(const FUniqueNetIdEOS& UserId) const;
 
diff --git a/Engine/Plugins/Online/OnlineSubsystemEOS/Source/OnlineSubsystemEOSPlus/Private/OnlineUserEOSPlus.cpp b/Engine/Plugins/Online/OnlineSubsystemEOS/Source/OnlineSubsystemEOSPlus/Private/OnlineUserEOSPlus.cpp
--- a/Engine/Plugins/Online/OnlineSubsystemEOS/Source/OnlineSubsystemEOSPlus/Private/OnlineUserEOSPlus.cpp
+++ b/Engine/Plugins/Online/OnlineSubsystemEOS/Source/OnlineSubsystemEOSPlus/Private/OnlineUserEOSPlus.cpp
@@ -861,10 +861,10 @@ FPlatformUserId FOnlineUserEOSPlus::GetPlatformUserIdFromUniqueNetId(const FUniq
 	return FPlatformUserId();
 }
 
-void FOnlineUserEOSPlus::GetLinkedAccountAuthToken(int32 LocalUserNum, const FOnGetLinkedAccountAuthTokenCompleteDelegate& Delegate) const
+void FOnlineUserEOSPlus::GetLinkedAccountAuthToken(int32 LocalUserNum, const FString& TokenType, const FOnGetLinkedAccountAuthTokenCompleteDelegate& Delegate) const
 {
 	// Pass through to the platform layer
-	BaseIdentityInterface->GetLinkedAccountAuthToken(LocalUserNum, Delegate);
+	BaseIdentityInterface->GetLinkedAccountAuthToken(LocalUserNum, TokenType, Delegate);
 }
 
 void FOnlineUserEOSPlus::OnFriendsChanged()
diff --git a/Engine/Plugins/Online/OnlineSubsystemEOS/Source/OnlineSubsystemEOSPlus/Private/OnlineUserEOSPlus.h b/Engine/Plugins/Online/OnlineSubsystemEOS/Source/OnlineSubsystemEOSPlus/Private/OnlineUserEOSPlus.h
--- a/Engine/Plugins/Online/OnlineSubsystemEOS/Source/OnlineSubsystemEOSPlus/Private/OnlineUserEOSPlus.h
+++ b/Engine/Plugins/Online/OnlineSubsystemEOS/Source/OnlineSubsystemEOSPlus/Private/OnlineUserEOSPlus.h
@@ -396,7 +396,7 @@ public:
 	virtual FString GetAuthType() const override;
 	virtual void RevokeAuthToken(const FUniqueNetId& LocalUserId, const FOnRevokeAuthTokenCompleteDelegate& Delegate) override;
 	virtual FPlatformUserId GetPlatformUserIdFromUniqueNetId(const FUniqueNetId& UniqueNetId) const override;
-	virtual void GetLinkedAccountAuthToken(int32 LocalUserNum, const FOnGetLinkedAccountAuthTokenCompleteDelegate& Delegate) const override;
+	virtual void GetLinkedAccountAuthToken(int32 LocalUserNum, const FString& TokenType, const FOnGetLinkedAccountAuthTokenCompleteDelegate& Delegate) const override;
 // ~IOnlineIdentity Interface
 
 // IOnlineFriends Interface
diff --git a/Engine/Plugins/Online/OnlineSubsystemSteam/Source/Private/OnlineIdentityInterfaceSteam.cpp b/Engine/Plugins/Online/OnlineSubsystemSteam/Source/Private/OnlineIdentityInterfaceSteam.cpp
--- a/Engine/Plugins/Online/OnlineSubsystemSteam/Source/Private/OnlineIdentityInterfaceSteam.cpp
+++ b/Engine/Plugins/Online/OnlineSubsystemSteam/Source/Private/OnlineIdentityInterfaceSteam.cpp
@@ -1,6 +1,8 @@
 // Copyright Epic Games, Inc. All Rights Reserved.
 
 #include "OnlineIdentityInterfaceSteam.h"
+
+#include "Misc/ConfigCacheIni.h"
 #include "OnlineSubsystemSteam.h"
 #include "OnlineSubsystemSteamTypes.h"
 #include "OnlineError.h"
@@ -239,9 +241,44 @@ FString FOnlineIdentitySteam::GetAuthType() const
 	return TEXT("");
 }
 
-void FOnlineIdentitySteam::GetLinkedAccountAuthToken(int32 LocalUserNum, const FOnGetLinkedAccountAuthTokenCompleteDelegate& Delegate) const
+void FOnlineIdentitySteam::GetLinkedAccountAuthToken(int32 LocalUserNum, const FString& InTokenType, const FOnGetLinkedAccountAuthTokenCompleteDelegate& Delegate) const
 {
-	FExternalAuthToken AuthToken;
-	AuthToken.TokenString = GetAuthToken(LocalUserNum);
-	Delegate.ExecuteIfBound(LocalUserNum, AuthToken.HasTokenString(), AuthToken);
+	const TCHAR* AppTokenType = TEXT("App");
+	const TCHAR* SessionTokenType = TEXT("Session");
+
+	FString TokenType = InTokenType;
+	if (TokenType.IsEmpty())
+	{
+		// Default to APP tokens
+		TokenType = AppTokenType;
+
+		// But allow the user to configure something else
+		GConfig->GetString(TEXT("OnlineSubsystemSteam"), TEXT("DefaultLinkedAccountAuthTokenType"), TokenType, GEngineIni);
+	}
+
+	if (TokenType == AppTokenType)
+	{
+		SteamSubsystem->GetEncryptedAppTicketInterface()->OnEncryptedAppTicketResultDelegate.AddLambda([this, LocalUserNum, OnComplete = FOnGetLinkedAccountAuthTokenCompleteDelegate(Delegate)](bool bEncryptedDataAvailable, int32 ResultCode)
+		{
+			FExternalAuthToken ExternalToken;
+			if (bEncryptedDataAvailable)
+			{
+				SteamSubsystem->GetEncryptedAppTicketInterface()->GetEncryptedAppTicket(ExternalToken.TokenData);
+			}
+			// Pass the info back to the original caller
+			OnComplete.ExecuteIfBound(LocalUserNum, ExternalToken.HasTokenData(), ExternalToken);
+		});
+		SteamSubsystem->GetEncryptedAppTicketInterface()->RequestEncryptedAppTicket(nullptr, 0);
+	}
+	else if (TokenType == SessionTokenType)
+	{
+		FExternalAuthToken AuthToken;
+		AuthToken.TokenString = GetAuthToken(LocalUserNum);
+		Delegate.ExecuteIfBound(LocalUserNum, AuthToken.HasTokenString(), AuthToken);
+	}
+	else
+	{	
+		UE_LOG_ONLINE_IDENTITY(Warning, TEXT("FOnlineIdentitySteam::GetLinkedAccountAuthToken TokenType=[%s] unknown"), *TokenType);
+		Delegate.ExecuteIfBound(LocalUserNum, false, FExternalAuthToken());
+	}
 }
diff --git a/Engine/Plugins/Online/OnlineSubsystemSteam/Source/Private/OnlineIdentityInterfaceSteam.h b/Engine/Plugins/Online/OnlineSubsystemSteam/Source/Private/OnlineIdentityInterfaceSteam.h
--- a/Engine/Plugins/Online/OnlineSubsystemSteam/Source/Private/OnlineIdentityInterfaceSteam.h
+++ b/Engine/Plugins/Online/OnlineSubsystemSteam/Source/Private/OnlineIdentityInterfaceSteam.h
@@ -46,7 +46,7 @@ public:
 	virtual void GetUserPrivilege(const FUniqueNetId& UserId, EUserPrivileges::Type Privilege, const FOnGetUserPrivilegeCompleteDelegate& Delegate) override;
 	virtual FPlatformUserId GetPlatformUserIdFromUniqueNetId(const FUniqueNetId& UniqueNetId) const override;
 	virtual FString GetAuthType() const override;
-	virtual void GetLinkedAccountAuthToken(int32 LocalUserNum, const FOnGetLinkedAccountAuthTokenCompleteDelegate& Delegate) const override;
+	virtual void GetLinkedAccountAuthToken(int32 LocalUserNum, const FString& TokenType, const FOnGetLinkedAccountAuthTokenCompleteDelegate& Delegate) const override;
 };
 
 typedef TSharedPtr<FOnlineIdentitySteam, ESPMode::ThreadSafe> FOnlineIdentitySteamPtr;
