PR #10108: Solve the problem that OwningUserName displays garbled code in Chinese, Japanese and other languages (Contributed by [REDACTED_PERSON]) (developer note: Steam specific issue)

#10108

diff --git a/Engine/Plugins/Online/OnlineSubsystemSteam/Source/Private/OnlineSessionAsyncLobbySteam.cpp b/Engine/Plugins/Online/OnlineSubsystemSteam/Source/Private/OnlineSessionAsyncLobbySteam.cpp
--- a/Engine/Plugins/Online/OnlineSubsystemSteam/Source/Private/OnlineSessionAsyncLobbySteam.cpp
+++ b/Engine/Plugins/Online/OnlineSubsystemSteam/Source/Private/OnlineSessionAsyncLobbySteam.cpp
@@ -299,9 +299,9 @@ bool FillSessionFromLobbyData(FOnlineSubsystemSteam* SteamSubsystem, const FUniq
 		}
 		else if (FCStringAnsi::Stricmp(Key, STEAMKEY_OWNINGUSERNAME) == 0)
 		{
-			if (FCString::Strlen(ANSI_TO_TCHAR(Value)) > 0)
+			if (FCString::Strlen(UTF8_TO_TCHAR(Value)) > 0)
 			{
-				Session.OwningUserName = Value;
+				Session.OwningUserName = UTF8_TO_TCHAR(Value);
 				KeysFound++;
 			}
 		}
diff --git a/Engine/Plugins/Online/OnlineSubsystemSteam/Source/Private/OnlineSessionInterfaceSteam.cpp b/Engine/Plugins/Online/OnlineSubsystemSteam/Source/Private/OnlineSessionInterfaceSteam.cpp
--- a/Engine/Plugins/Online/OnlineSubsystemSteam/Source/Private/OnlineSessionInterfaceSteam.cpp
+++ b/Engine/Plugins/Online/OnlineSubsystemSteam/Source/Private/OnlineSessionInterfaceSteam.cpp
@@ -234,7 +234,7 @@ bool FOnlineSessionSteam::CreateSession(int32 HostingPlayerNum, FName SessionNam
 
 		Session->HostingPlayerNum = HostingPlayerNum;
 		Session->OwningUserId = SteamUser() ? FUniqueNetIdSteamPtr(FUniqueNetIdSteam::Create(SteamUser()->GetSteamID())) : nullptr;
-		Session->OwningUserName = SteamFriends() ? SteamFriends()->GetPersonaName() : GetCustomDedicatedServerName();
+		Session->OwningUserName = SteamFriends() ? UTF8_TO_TCHAR(SteamFriends()->GetPersonaName()) : GetCustomDedicatedServerName();
 		
 		// Unique identifier of this build for compatibility
 		Session->SessionSettings.BuildUniqueId = GetBuildUniqueId();
