Fixed compilation error due incorrect define resolving to protected instead of public

diff --git a/Engine/Plugins/Online/OnlineSubsystemSteam/Source/OnlineSubsystemSteam.Build.cs b/Engine/Plugins/Online/OnlineSubsystemSteam/Source/OnlineSubsystemSteam.Build.cs
--- a/Engine/Plugins/Online/OnlineSubsystemSteam/Source/OnlineSubsystemSteam.Build.cs
+++ b/Engine/Plugins/Online/OnlineSubsystemSteam/Source/OnlineSubsystemSteam.Build.cs
@@ -8,6 +8,8 @@ public class OnlineSubsystemSteam : ModuleRules
 	public OnlineSubsystemSteam(ReadOnlyTargetRules Target) : base(Target)
 	{
 		PrivateDefinitions.Add("ONLINESUBSYSTEMSTEAM_PACKAGE=1");
+		PrivateDefinitions.Add("STEAMSHARED_PACKAGE=1");
+
 		PCHUsage = ModuleRules.PCHUsageMode.UseExplicitOrSharedPCHs;
 
 		PublicDependencyModuleNames.AddRange(
