Steam: Do not disable the sandbox because it can cause async file operations to randomly fail

diff --git a/Engine/Plugins/Online/OnlineSubsystemSteam/Source/Private/OnlineSubsystemSteam.cpp b/Engine/Plugins/Online/OnlineSubsystemSteam/Source/Private/OnlineSubsystemSteam.cpp
--- a/Engine/Plugins/Online/OnlineSubsystemSteam/Source/Private/OnlineSubsystemSteam.cpp
+++ b/Engine/Plugins/Online/OnlineSubsystemSteam/Source/Private/OnlineSubsystemSteam.cpp
@@ -4,8 +4,8 @@
 #include "HAL/FileManager.h"
 #include "HAL/PlatformFile.h"
 #include "HAL/RunnableThread.h"
-#include "Serialization/BufferArchive.h"
 #include "Interfaces/IPluginManager.h"
+#include "Misc/StringBuilder.h"
 
 #include "SocketSubsystemSteam.h"
 
@@ -98,27 +98,6 @@ static void __cdecl SteamworksWarningMessageHookNoOp(int Severity, const char *M
 	// no-op.
 }
 
-class FScopeSandboxContext
-{
-private:
-	/** Previous state of sandbox enable */
-	bool bSandboxWasEnabled;
-
-	FScopeSandboxContext() {}
-
-public:
-	FScopeSandboxContext(bool bSandboxEnabled)
-	{
-		bSandboxWasEnabled = IFileManager::Get().IsSandboxEnabled();
-		IFileManager::Get().SetSandboxEnabled(bSandboxEnabled);
-	}
-
-	~FScopeSandboxContext()
-	{
-		IFileManager::Get().SetSandboxEnabled(bSandboxWasEnabled);
-	}
-};
-
 inline FString GetSteamAppIdFilename()
 {
 	return FString::Printf(TEXT("%s%s"), FPlatformProcess::BaseDir(), STEAMAPPIDFILENAME);
@@ -133,12 +112,9 @@ static bool WriteSteamAppIdToDisk(int32 SteamAppId)
 {
 	if (SteamAppId > 0)
 	{
-		// Turn off sandbox temporarily to make sure file is where it's always expected
-		FScopeSandboxContext ScopedSandbox(false);
-
 		// Access the physical file writer directly so that we still write next to the executable in CotF builds.
 		FString SteamAppIdFilename = GetSteamAppIdFilename();
-		IFileHandle* Handle = IPlatformFile::GetPlatformPhysical().OpenWrite(*SteamAppIdFilename, false, false);
+		TUniquePtr<IFileHandle> Handle(IPlatformFile::GetPlatformPhysical().OpenWrite(*SteamAppIdFilename, false, false));
 		if (!Handle)
 		{
 			UE_LOG_ONLINE(Error, TEXT("Failed to create file: %s"), *SteamAppIdFilename);
@@ -146,15 +122,8 @@ static bool WriteSteamAppIdToDisk(int32 SteamAppId)
 		}
 		else
 		{
-			FString AppId = FString::Printf(TEXT("%d"), SteamAppId);
-
-			FBufferArchive Archive;
-			Archive.Serialize((void*)TCHAR_TO_ANSI(*AppId), AppId.Len());
-
-			Handle->Write(Archive.GetData(), Archive.Num());
-			delete Handle;
-			Handle = nullptr;
-
+			TAnsiStringBuilder<16> AppId(InPlace, SteamAppId);
+			Handle->Write((const uint8*)AppId.GetData(), AppId.Len());
 			return true;
 		}
 	}
@@ -169,11 +138,10 @@ static bool WriteSteamAppIdToDisk(int32 SteamAppId)
 static void DeleteSteamAppIdFromDisk()
 {
 	const FString SteamAppIdFilename = GetSteamAppIdFilename();
-	// Turn off sandbox temporarily to make sure file is where it's always expected
-	FScopeSandboxContext ScopedSandbox(false);
-	if (FPaths::FileExists(SteamAppIdFilename))
+	IPlatformFile& Physical = IPlatformFile::GetPlatformPhysical();
+	if (Physical.FileExists(*SteamAppIdFilename))
 	{
-		bool bSuccessfullyDeleted = IFileManager::Get().Delete(*SteamAppIdFilename);
+		Physical.DeleteFile(*SteamAppIdFilename);
 	}
 }
 
