PR #7028: Updated the Steam OSS to have proper presence for users.

diff --git a/Engine/Plugins/Online/OnlineSubsystemSteam/Source/Private/OnlineFriendsInterfaceSteam.cpp b/Engine/Plugins/Online/OnlineSubsystemSteam/Source/Private/OnlineFriendsInterfaceSteam.cpp
--- a/Engine/Plugins/Online/OnlineSubsystemSteam/Source/Private/OnlineFriendsInterfaceSteam.cpp
+++ b/Engine/Plugins/Online/OnlineSubsystemSteam/Source/Private/OnlineFriendsInterfaceSteam.cpp
@@ -1,12 +1,15 @@
 // Copyright Epic Games, Inc. All Rights Reserved.
 
 #include "OnlineFriendsInterfaceSteam.h"
+#include "OnlinePresenceInterfaceSteam.h"
 #include "OnlineSubsystemSteam.h"
 #include "OnlineError.h"
 #include "OnlineSubsystemSteamTypes.h"
 #include <steam/isteamuser.h>
 
 // FOnlineFriendSteam
+const FOnlineUserPresence FOnlineFriendSteam::EmptyPresence;
+
 FOnlineFriendSteam::FOnlineFriendSteam(const CSteamID& InUserId)
 	: UserId(FUniqueNetIdSteam::Create(InUserId))
 {
@@ -43,7 +46,34 @@ EInviteStatus::Type FOnlineFriendSteam::GetInviteStatus() const
 
 const FOnlineUserPresence& FOnlineFriendSteam::GetPresence() const
 {
-	return Presence;
+	// Get the steam subsystem
+	if (IOnlineSubsystem* SteamOSS = IOnlineSubsystem::Get(STEAM_SUBSYSTEM))
+	{
+		// Get the steam presence interface
+		if (FOnlinePresenceSteam* OnlinePresence = static_cast<FOnlinePresenceSteam*>(SteamOSS->GetPresenceInterface().Get()))
+		{
+			// Get the presence if it's available
+			TSharedPtr<FOnlineUserPresence> Presence;
+			EOnlineCachedResult::Type Result = OnlinePresence->GetCachedPresence(*UserId, Presence);
+
+			// If the user's presence was not cached (it might be the first time)
+			if (Result == EOnlineCachedResult::Type::NotFound)
+			{
+				// Add the user's presence to the cache
+				OnlinePresence->UpdatePresenceForUser(*UserId);
+
+				// Try getting the user's presence again
+				OnlinePresence->GetCachedPresence(*UserId, Presence);
+			}
+
+			if(Presence.IsValid())
+			{
+				return *Presence;
+			}
+		}
+	}
+
+	return EmptyPresence;
 }
 
 // FOnlineFriendsStream
@@ -292,34 +322,9 @@ void FOnlineAsyncTaskSteamReadFriendsList::Finalize()
 
 			// Now fill in the friend info
 			Friend->AccountData.Add(TEXT("nickname"), NickName);
-			Friend->Presence.Status.StatusStr = UTF8_TO_TCHAR(SteamFriendsPtr->GetFriendRichPresence(SteamPlayerId,"status"));
-			Friend->Presence.bIsJoinable = bInASession;
-			Friend->Presence.bIsOnline = bIsOnline;
-			Friend->Presence.bIsPlaying = bIsPlayingAGame;
-			Friend->Presence.bIsPlayingThisGame = bIsPlayingThisGame;
 
-			switch (SteamFriendsPtr->GetFriendPersonaState(SteamPlayerId))
-			{
-				case k_EPersonaStateOffline:
-					Friend->Presence.Status.State = EOnlinePresenceState::Offline;
-					break;
-				case k_EPersonaStateBusy:
-					Friend->Presence.Status.State = EOnlinePresenceState::DoNotDisturb;
-					break;
-				case k_EPersonaStateAway:
-					Friend->Presence.Status.State = EOnlinePresenceState::Away;
-					break;
-				case k_EPersonaStateSnooze:
-					Friend->Presence.Status.State = EOnlinePresenceState::ExtendedAway;
-					break;
-				default:
-					Friend->Presence.Status.State = EOnlinePresenceState::Online;
-					break;
-			}
-			// Remote friend is responsible for updating their presence to have the voice flag
-			FString VoicePresenceString = UTF8_TO_TCHAR(SteamFriendsPtr->GetFriendRichPresence(SteamPlayerId,"HasVoice"));
-			// Determine if the user has voice support
-			Friend->Presence.bHasVoiceSupport = VoicePresenceString == TEXT("true");
+			// Initially update/cache the user's presence
+			Friend->GetPresence();
 		}
 	}
 }
diff --git a/Engine/Plugins/Online/OnlineSubsystemSteam/Source/Private/OnlineFriendsInterfaceSteam.h b/Engine/Plugins/Online/OnlineSubsystemSteam/Source/Private/OnlineFriendsInterfaceSteam.h
--- a/Engine/Plugins/Online/OnlineSubsystemSteam/Source/Private/OnlineFriendsInterfaceSteam.h
+++ b/Engine/Plugins/Online/OnlineSubsystemSteam/Source/Private/OnlineFriendsInterfaceSteam.h
@@ -62,8 +62,9 @@ public:
 	FUniqueNetIdRef UserId;
 	/** Any addition account data associated with the friend */
 	TMap<FString, FString> AccountData;
-	/** @temp presence info */
-	FOnlineUserPresence Presence;
+
+	/** A backup presence we can use if we are unable to get a friend's presence for any reason. */
+	static const FOnlineUserPresence EmptyPresence;
 };
 
 /**
diff --git a/Engine/Plugins/Online/OnlineSubsystemSteam/Source/Private/OnlinePresenceInterfaceSteam.cpp b/Engine/Plugins/Online/OnlineSubsystemSteam/Source/Private/OnlinePresenceInterfaceSteam.cpp
--- a/Engine/Plugins/Online/OnlineSubsystemSteam/Source/Private/OnlinePresenceInterfaceSteam.cpp
+++ b/Engine/Plugins/Online/OnlineSubsystemSteam/Source/Private/OnlinePresenceInterfaceSteam.cpp
@@ -55,6 +55,15 @@ void FOnlineUserPresenceSteam::Update(const FUniqueNetIdSteam& FriendId)
 		SessionId = FUniqueNetIdSteam::Create(FriendInfo.m_steamIDLobby); 
 		bIsPlayingThisGame = (FriendInfo.m_gameID.AppID() == SteamUtils()->GetAppID());
 	}
+	else
+	{
+		bIsPlaying = false;
+		SessionId = nullptr;
+		bIsPlayingThisGame = false;
+	}
+
+	bIsJoinable = false;
+	bHasVoiceSupport = false;
 
 	// Processing presence
 	for (int32 RPIdx = 0; RPIdx < SteamFriendPtr->GetFriendRichPresenceKeyCount(FriendId); ++RPIdx)
@@ -66,6 +75,7 @@ void FOnlineUserPresenceSteam::Update(const FUniqueNetIdSteam& FriendId)
 		if (Key == DefaultSteamConnectionKey)
 		{
 			bIsJoinable = true;
+			continue;
 		}
 
 		// This key has it's own entry, there's no need to push it to the property field
@@ -75,6 +85,13 @@ void FOnlineUserPresenceSteam::Update(const FUniqueNetIdSteam& FriendId)
 			continue;
 		}
 
+		// Determine if the user has voice support
+		if (Key == "HasVoice" && Value == "true")
+		{
+			bHasVoiceSupport = true;
+			continue;
+		}
+
 		FVariantData PropertyData;
 		PropertyData.SetValue(Value);
 		Status.Properties.Add(Key, PropertyData);
