PR #7301: Fix -nosteam command line

Epic note: -nosteam was previously functional but some modules were still loaded and dormant, this change will stop them from being registered for games that may have issues with the modules being loaded.

diff --git a/Engine/Plugins/Online/OnlineSubsystemSteam/Source/Private/OnlineSubsystemModuleSteam.cpp b/Engine/Plugins/Online/OnlineSubsystemSteam/Source/Private/OnlineSubsystemModuleSteam.cpp
--- a/Engine/Plugins/Online/OnlineSubsystemSteam/Source/Private/OnlineSubsystemModuleSteam.cpp
+++ b/Engine/Plugins/Online/OnlineSubsystemSteam/Source/Private/OnlineSubsystemModuleSteam.cpp
@@ -3,6 +3,7 @@
 #include "OnlineSubsystemSteamModule.h"
 #include "OnlineSubsystemSteam.h"
 #include "SteamSharedModule.h"
+#include "Misc/CommandLine.h"
 
 IMPLEMENT_MODULE(FOnlineSubsystemSteamModule, OnlineSubsystemSteam);
 
@@ -69,6 +70,12 @@ FOnlineSubsystemSteamPtr FOnlineFactorySteam::SteamSingleton = nullptr;
 
 void FOnlineSubsystemSteamModule::StartupModule()
 {
+	if (FParse::Param(FCommandLine::Get(), TEXT("nosteam")))
+	{
+		UE_LOG_ONLINE(Warning, TEXT("Steam subsystem has been disabled by command line (-nosteam)"));
+		return;
+	}
+	
 	FSteamSharedModule& SharedModule = FSteamSharedModule::Get();
 
 	// Load the Steam modules before first call to API
