Acceptable words fixes for Steam files

diff --git a/Engine/Plugins/Online/OnlineSubsystemSteam/Source/Private/OnlineAsyncTaskManagerSteam.cpp b/Engine/Plugins/Online/OnlineSubsystemSteam/Source/Private/OnlineAsyncTaskManagerSteam.cpp
--- a/Engine/Plugins/Online/OnlineSubsystemSteam/Source/Private/OnlineAsyncTaskManagerSteam.cpp
+++ b/Engine/Plugins/Online/OnlineSubsystemSteam/Source/Private/OnlineAsyncTaskManagerSteam.cpp
@@ -611,7 +611,7 @@ void FOnlineAsyncTaskManagerSteam::OnSteamServersDisconnected(SteamServersDiscon
 }
 
 /**
- * Notification event from Steam that server session has connected with the master server
+ * Notification event from Steam that server session has connected with the server list
  */
 class FOnlineAsyncEventSteamServerConnectedGS : public FOnlineAsyncEvent<FOnlineSubsystemSteam>
 {
@@ -683,7 +683,7 @@ void FOnlineAsyncTaskManagerSteam::OnSteamServersConnectedGS(SteamServersConnect
 }
 
 /**
- * Notification event from Steam that server session has been disconnected with the master server
+ * Notification event from Steam that server session has been disconnected with the server list
  */
 class FOnlineAsyncEventSteamServerDisconnectedGS : public FOnlineAsyncEvent<FOnlineSubsystemSteam>
 {
diff --git a/Engine/Plugins/Online/OnlineSubsystemSteam/Source/Private/OnlineSessionAsyncServerSteam.cpp b/Engine/Plugins/Online/OnlineSubsystemSteam/Source/Private/OnlineSessionAsyncServerSteam.cpp
--- a/Engine/Plugins/Online/OnlineSubsystemSteam/Source/Private/OnlineSessionAsyncServerSteam.cpp
+++ b/Engine/Plugins/Online/OnlineSubsystemSteam/Source/Private/OnlineSessionAsyncServerSteam.cpp
@@ -317,7 +317,7 @@ void UpdatePublishedSettings(UWorld* World, FNamedOnlineSession* Session)
 	// Generate the tag data first, due to its small size
 	for (FSteamSessionKeyValuePairs::TConstIterator It(ViaOnlineServicePairs); It; ++It)
 	{
-		UE_LOG_ONLINE_SESSION(Verbose, TEXT("Master Server Game Tags (%s, %s)"), *It.Key(), *It.Value());
+		UE_LOG_ONLINE_SESSION(Verbose, TEXT("Server List Game Tags (%s, %s)"), *It.Key(), *It.Value());
 		FString NewKey = FString::Printf(TEXT(",%s:%s"), *It.Key(), *It.Value());
 		if (GameTagsString.Len() + NewKey.Len() < k_cbMaxGameServerTags)
 		{
@@ -343,7 +343,7 @@ void UpdatePublishedSettings(UWorld* World, FNamedOnlineSession* Session)
 	FString GameDataString;
 	for (FSteamSessionKeyValuePairs::TConstIterator It(ViaOnlineServicePairs); It; ++It)
 	{
-		UE_LOG_ONLINE_SESSION(Verbose, TEXT("Master Server Game Data (%s, %s)"), *It.Key(), *It.Value());
+		UE_LOG_ONLINE_SESSION(Verbose, TEXT("Server List Game Data (%s, %s)"), *It.Key(), *It.Value());
 		FString NewKey = FString::Printf(TEXT("%s:%s"), *It.Key(), *It.Value());
 
 		if (GameDataString.Len() + NewKey.Len() < k_cbMaxGameServerGameData)
@@ -363,14 +363,14 @@ void UpdatePublishedSettings(UWorld* World, FNamedOnlineSession* Session)
 		SteamGameServerPtr->SetGameTags(TCHAR_TO_UTF8(*GameTagsString));
 	}
 
-	// Large and searchable game data (never returned, used to filter things on the master server)
+	// Large and searchable game data (never returned, used to filter things on the server list)
 	if (GameDataString.Len() > 0 && GameDataString.Len() < k_cbMaxGameServerGameData)
 	{
 		UE_LOG_ONLINE_SESSION(Verbose, TEXT("SetGameData(%s)"), *GameDataString);
 		SteamGameServerPtr->SetGameData(TCHAR_TO_UTF8(*GameDataString));
 	}
 
-	// Set the advertised filter keys (these can not be filtered at master-server level, only client side via talking to the server)
+	// Set the advertised filter keys (these can not be filtered at server list level, only client side via talking to the server)
 	SteamGameServerPtr->ClearAllKeyValues();
 
 	// Push all the filters in for the KeyValue data.
@@ -570,7 +570,7 @@ void FOnlineAsyncTaskSteamUpdateServer::Tick()
 		{
 			UWorld* World = GetWorldForOnline(Subsystem->GetInstanceName());
 
-			// Master server update
+			// Server list update
 			UpdatePublishedSettings(World, Session);
 		}
 
@@ -613,7 +613,7 @@ void FOnlineAsyncTaskSteamLogoffServer::Tick()
 	if (!bInit)
 	{
 		// @TODO ONLINE Listen Servers need to unset rich presence
-		//SteamFriends()->SetRichPresence("connect", ""); for master server sessions
+		//SteamFriends()->SetRichPresence("connect", ""); for listed sessions
 		SteamGameServer()->SetAdvertiseServerActive(false);  // renamed EnableHeartbeats to SetAdvertiseServerActive in 1.53
 		SteamGameServer()->LogOff();
 		bInit = true;
@@ -851,7 +851,7 @@ void FPendingSearchResultSteam::CancelQuery()
 #endif
 
 /**
- *  Create the proper query for the master server based on the given search settings
+ *  Create the proper query for the server list based on the given search settings
  *
  * @param OutFilter Steam structure containing the proper filters
  * @param NumFilters number of filters contained in the array above
diff --git a/Engine/Plugins/Online/OnlineSubsystemSteam/Source/Private/OnlineSessionAsyncServerSteam.h b/Engine/Plugins/Online/OnlineSubsystemSteam/Source/Private/OnlineSessionAsyncServerSteam.h
--- a/Engine/Plugins/Online/OnlineSubsystemSteam/Source/Private/OnlineSessionAsyncServerSteam.h
+++ b/Engine/Plugins/Online/OnlineSubsystemSteam/Source/Private/OnlineSessionAsyncServerSteam.h
@@ -257,7 +257,7 @@ PACKAGE_SCOPE:
 	TIndirectArray<FPendingSearchResultSteam> PendingSearchResults;
 	/** Search settings specified for the query */
 	TSharedPtr<class FOnlineSessionSearch> SearchSettings;
-	/** Master server request handle */
+	/** Server list request handle */
 	HServerListRequest ServerListRequestHandle;
 
 public:
@@ -274,7 +274,7 @@ public:
 	}
 
 	/**
-     *  Create the proper query for the master server based on the given search settings
+     *  Create the proper query for the server list based on the given search settings
      *
      * @param OutFilter Steam structure containing the proper filters
      * @param NumFilters number of filters contained in the array above
@@ -412,7 +412,7 @@ public:
 };
 
 /**
- *	Turns a friends accepted invite request into a valid search result (master server version)
+ *	Turns a friends accepted invite request into a valid search result (server list version)
  */
 class FOnlineAsyncEventSteamInviteAccepted : public FOnlineAsyncEvent<FOnlineSubsystemSteam>
 {
diff --git a/Engine/Plugins/Online/OnlineSubsystemSteam/Source/Private/OnlineSessionInterfaceSteam.cpp b/Engine/Plugins/Online/OnlineSubsystemSteam/Source/Private/OnlineSessionInterfaceSteam.cpp
--- a/Engine/Plugins/Online/OnlineSubsystemSteam/Source/Private/OnlineSessionInterfaceSteam.cpp
+++ b/Engine/Plugins/Online/OnlineSubsystemSteam/Source/Private/OnlineSessionInterfaceSteam.cpp
@@ -321,8 +321,8 @@ uint32 FOnlineSessionSteam::CreateInternetSession(int32 HostingPlayerNum, FNamed
 	uint32 Result = ONLINE_FAIL;
 
 	// Only allowed one published session with Steam
-	FNamedOnlineSession* MasterSession = GetGameServerSession();
-	if (MasterSession == nullptr)
+	FNamedOnlineSession* ListedSession = GetGameServerSession();
+	if (ListedSession == nullptr)
 	{
 		if (SteamSubsystem->IsSteamServerAvailable())
 		{
@@ -660,7 +660,7 @@ uint32 FOnlineSessionSteam::DestroyInternetSession(FNamedOnlineSession* Session,
 
 	if (bSteamworksGameServerConnected && GameServerSteamId->IsValid())
 	{
-		// Logoff the master server
+		// Remove this server from the server list
 		FOnlineAsyncTaskSteamLogoffServer* LogoffTask = new FOnlineAsyncTaskSteamLogoffServer(SteamSubsystem, Session->SessionName);
 		SteamSubsystem->QueueAsyncTask(LogoffTask);
 	}
@@ -828,7 +828,7 @@ bool FOnlineSessionSteam::CancelFindSessions()
 		}
 		else
 		{
-			// @TODO ONLINE Master Server Version
+			// @TODO ONLINE Server list version
 			Return = ONLINE_SUCCESS;
 			// There is no CANCEL lobby query
 			// NULLing out the object will prevent the async event from adding the results
diff --git a/Engine/Plugins/Online/OnlineSubsystemSteam/Source/Private/OnlineSubsystemSteamTypes.h b/Engine/Plugins/Online/OnlineSubsystemSteam/Source/Private/OnlineSubsystemSteamTypes.h
--- a/Engine/Plugins/Online/OnlineSubsystemSteam/Source/Private/OnlineSubsystemSteamTypes.h
+++ b/Engine/Plugins/Online/OnlineSubsystemSteam/Source/Private/OnlineSubsystemSteamTypes.h
@@ -16,7 +16,7 @@ namespace ESteamSession
 		None,
 		/** Session managed as a lobby on backend */
 		LobbySession,
-		/** Session managed by master server publishing */
+		/** Session managed by the server list */
 		AdvertisedSessionHost,
 		/** Session client of a game server session */
 		AdvertisedSessionClient,
diff --git a/Engine/Plugins/Runtime/Steam/SteamShared/Source/SteamShared/Private/SteamSharedModule.cpp b/Engine/Plugins/Runtime/Steam/SteamShared/Source/SteamShared/Private/SteamSharedModule.cpp
--- a/Engine/Plugins/Runtime/Steam/SteamShared/Source/SteamShared/Private/SteamSharedModule.cpp
+++ b/Engine/Plugins/Runtime/Steam/SteamShared/Source/SteamShared/Private/SteamSharedModule.cpp
@@ -325,7 +325,7 @@ FSteamServerInstanceHandler::FSteamServerInstanceHandler(FSteamSharedModule* Ste
 		}
 	}
 
-	// Allow the command line to override the default query port for master server communications
+	// Allow the command line to override the default query port for server list queries
 	if (FParse::Value(FCommandLine::Get(), TEXT("QueryPort="), QueryPort) == false)
 	{
 		if (!GConfig->GetInt(TEXT("OnlineSubsystemSteam"), TEXT("GameServerQueryPort"), QueryPort, GEngineIni))
@@ -363,7 +363,7 @@ FSteamServerInstanceHandler::FSteamServerInstanceHandler(FSteamSharedModule* Ste
 
 void FSteamServerInstanceHandler::InternalShutdown()
 {
-	// Log off of the steam master servers, removes this server immediately from the backend.
+	// Log off of the steam coordinator servers, removes this server immediately from the backend.
 	if (SteamGameServer() && SteamGameServer()->BLoggedOn())
 	{
 		SteamGameServer()->LogOff();
