Fixes for Steam's LeftLobby and IsMemberOfLobby, community contributed. Use a manual array search instead of built-in TArray functions since they behave incorrectly with shared refs

diff --git a/Engine/Plugins/Online/OnlineSubsystemSteam/Source/Private/OnlineSessionInterfaceSteam.h b/Engine/Plugins/Online/OnlineSubsystemSteam/Source/Private/OnlineSessionInterfaceSteam.h
--- a/Engine/Plugins/Online/OnlineSubsystemSteam/Source/Private/OnlineSessionInterfaceSteam.h
+++ b/Engine/Plugins/Online/OnlineSubsystemSteam/Source/Private/OnlineSessionInterfaceSteam.h
@@ -394,7 +394,15 @@ PACKAGE_SCOPE:
 	void LeftLobby(const FUniqueNetIdSteam& LobbyId)
 	{
 		FScopeLock ScopeLock(&JoinedLobbyLock);
-		JoinedLobbyList.RemoveSingleSwap(LobbyId.AsShared());
+
+		for (int i = 0; i < JoinedLobbyList.Num(); i++)
+		{
+			if (*JoinedLobbyList[i] == LobbyId)
+			{
+				JoinedLobbyList.RemoveAt(i);
+				return;
+			}
+		}
 	}
 
 	/**
@@ -407,7 +415,15 @@ PACKAGE_SCOPE:
 	bool IsMemberOfLobby(const FUniqueNetIdSteam& LobbyId)
 	{
 		FScopeLock ScopeLock(&JoinedLobbyLock);
-		return JoinedLobbyList.Find(LobbyId.AsShared()) != INDEX_NONE;
+
+		for (int i = 0; i < JoinedLobbyList.Num(); i++)
+		{
+			if (*JoinedLobbyList[i] == LobbyId)
+			{
+				return true;
+			}
+		}
+		return false;
 	}
 
 	/**
