PR #7291: Updated Steam's rich presence keys and implemented friends grouping
Epic note: Added a killswitch

diff --git a/Engine/Plugins/Online/OnlineSubsystemSteam/Source/Private/OnlinePresenceInterfaceSteam.cpp b/Engine/Plugins/Online/OnlineSubsystemSteam/Source/Private/OnlinePresenceInterfaceSteam.cpp
--- a/Engine/Plugins/Online/OnlineSubsystemSteam/Source/Private/OnlinePresenceInterfaceSteam.cpp
+++ b/Engine/Plugins/Online/OnlineSubsystemSteam/Source/Private/OnlinePresenceInterfaceSteam.cpp
@@ -4,12 +4,19 @@
 #include "OnlineSubsystemSteam.h"
 #include "OnlineSessionInterfaceSteam.h"
 #include <steam/isteamutils.h>
+#include "Misc/ConfigCacheIni.h"
 
 /** The default key that will contain the launch parameters for joining the game */
 const FString DefaultSteamConnectionKey = TEXT("connect");
 
 /** The default key that will update presence text in the platform's UI */
-const FString DefaultSteamPresenceKey = TEXT("status");
+const FString DefaultSteamPresenceKey = TEXT("steam_display");
+
+/** The default key that will group players in the platform's UI */
+const char* DefaultSteamPlayerGroupingKey = "steam_player_group";
+
+/** The default key that specifies the amount of players are in the party for the platform's UI */
+const char* DefaultSteamPlayerGroupSizeKey = "steam_player_group_size";
 
 void FOnlineUserPresenceSteam::Update(const FUniqueNetIdSteam& FriendId)
 {
@@ -159,6 +166,23 @@ void FOnlinePresenceSteam::SetPresence(const FUniqueNetId& User, const FOnlineUs
 				UE_LOG_ONLINE_PRESENCE(Warning, TEXT("Could not push the connection information to Steam"));
 			}
 		}
+
+		bool bDisableSteamAdvancedPresence = false;
+		GConfig->GetBool(TEXT("OnlineSubsystemSteam"), TEXT("DisableAdvancedPresence"), bDisableSteamAdvancedPresence, GEngineIni);
+		if (CurrentSession != NULL && CurrentSession->SessionSettings.bUsesPresence && !bDisableSteamAdvancedPresence)
+		{
+			FString SteamPlayerGroupingKeyString = CurrentSession->GetSessionIdStr();
+			if (!SteamPlayerGroupingKeyString.IsEmpty() && !SteamFriendsPtr->SetRichPresence(DefaultSteamPlayerGroupingKey, TCHAR_TO_UTF8(*SteamPlayerGroupingKeyString)))
+			{
+				UE_LOG_ONLINE_PRESENCE(Warning, TEXT("Could not push the player grouping key information to Steam"));
+			}
+
+			FString SteamPlayerGroupSizeString = FString::FromInt(CurrentSession->RegisteredPlayers.Num());
+			if (!SteamPlayerGroupSizeString.IsEmpty() && !SteamFriendsPtr->SetRichPresence(DefaultSteamPlayerGroupSizeKey, TCHAR_TO_UTF8(*SteamPlayerGroupSizeString)))
+			{
+				UE_LOG_ONLINE_PRESENCE(Warning, TEXT("Could not push the player group size information to Steam"));
+			}
+		}
 	}
 
 	// Push extra presence properties
