PR #9483: Add the ability to specify Steam OSS server names via config files

diff --git a/Engine/Plugins/Online/OnlineSubsystemSteam/Source/Private/OnlineSessionInterfaceSteam.cpp b/Engine/Plugins/Online/OnlineSubsystemSteam/Source/Private/OnlineSessionInterfaceSteam.cpp
--- a/Engine/Plugins/Online/OnlineSubsystemSteam/Source/Private/OnlineSessionInterfaceSteam.cpp
+++ b/Engine/Plugins/Online/OnlineSubsystemSteam/Source/Private/OnlineSessionInterfaceSteam.cpp
@@ -1338,19 +1338,19 @@ FString FOnlineSessionSteam::GetCustomDedicatedServerName() const
 {
 	FString ServerName;
 
-	if (FParse::Value(FCommandLine::Get(), TEXT("-SteamServerName="), ServerName))
+	if (!FParse::Value(FCommandLine::Get(), TEXT("-SteamServerName="), ServerName))
 	{
-		if (ServerName.Len() >= k_cbMaxGameServerName)
-		{
-			UE_LOG_ONLINE_SESSION(Warning, TEXT("SteamServerName overflows the maximum amount of characters %d allowed, truncating."), k_cbMaxGameServerName);
-			// Must have space for the null terminator
-			ServerName.LeftInline(k_cbMaxGameServerName - 1);
-		}
+		GConfig->GetString(TEXT("OnlineSubsystemSteam"), TEXT("SteamServerName"), ServerName, GEngineIni);
+	}
 
-		return ServerName;
+	if (ServerName.Len() >= k_cbMaxGameServerName)
+	{
+		UE_LOG_ONLINE_SESSION(Warning, TEXT("SteamServerName overflows the maximum amount of characters %d allowed, truncating."), k_cbMaxGameServerName);
+		// Must have space for the null terminator
+		ServerName.LeftInline(k_cbMaxGameServerName - 1);
 	}
-	
-	return TEXT("");
+
+	return ServerName;
 }
 
 FUniqueNetIdPtr FOnlineSessionSteam::CreateSessionIdFromString(const FString& SessionIdStr)
