PR #9491: Add some config vars for Steam OSS timeouts

diff --git a/Engine/Plugins/Online/OnlineSubsystemSteam/Source/Private/OnlineSessionAsyncLobbySteam.cpp b/Engine/Plugins/Online/OnlineSubsystemSteam/Source/Private/OnlineSessionAsyncLobbySteam.cpp
--- a/Engine/Plugins/Online/OnlineSubsystemSteam/Source/Private/OnlineSessionAsyncLobbySteam.cpp
+++ b/Engine/Plugins/Online/OnlineSubsystemSteam/Source/Private/OnlineSessionAsyncLobbySteam.cpp
@@ -8,6 +8,7 @@
 #include "IPAddressSteam.h"
 #include "SteamSessionKeys.h"
 #include "SteamUtilities.h"
+#include "Misc/ConfigCacheIni.h"
 
 /** 
  * Helper function to convert enums 
@@ -1039,13 +1040,16 @@ void FOnlineAsyncTaskSteamFindLobbiesBase::Tick()
 		{
 			FOnlineSessionSteamPtr SessionInt = StaticCastSharedPtr<FOnlineSessionSteam>(Subsystem->GetSessionInterface());
 
+			float AsyncTimeout = ASYNC_TASK_TIMEOUT;
+			GConfig->GetFloat(TEXT("OnlineSubsystemSteam"), TEXT("OnlineAsyncTaskSteamFindLobbiesBaseTimeout"), AsyncTimeout, GEngineIni);
+
 			// Waiting for the lobby updates to fill in
 			if (LobbyIDs.Num() == SessionInt->PendingSearchLobbyIds.Num())
 			{
 				FindLobbiesState = EFindLobbiesState::Finished;
 			}
 			// Fallback timeout in case we don't hear from Steam
-			else if (GetElapsedTime() >= ASYNC_TASK_TIMEOUT)
+			else if (GetElapsedTime() >= AsyncTimeout)
 			{
 				bWasSuccessful = false;
 				FindLobbiesState = EFindLobbiesState::Finished;
diff --git a/Engine/Plugins/Online/OnlineSubsystemSteam/Source/Private/OnlineSessionAsyncServerSteam.cpp b/Engine/Plugins/Online/OnlineSubsystemSteam/Source/Private/OnlineSessionAsyncServerSteam.cpp
--- a/Engine/Plugins/Online/OnlineSubsystemSteam/Source/Private/OnlineSessionAsyncServerSteam.cpp
+++ b/Engine/Plugins/Online/OnlineSubsystemSteam/Source/Private/OnlineSessionAsyncServerSteam.cpp
@@ -440,8 +440,11 @@ void FOnlineAsyncTaskSteamCreateServer::Tick()
 	}
 	else
 	{
+		float AsyncTimeout = ASYNC_TASK_TIMEOUT;
+		GConfig->GetFloat(TEXT("OnlineSubsystemSteam"), TEXT("OnlineAsyncTaskSteamCreateServerTimeout"), AsyncTimeout, GEngineIni);
+		
 		// Fallback timeout in case we don't hear from Steam
-		if (GetElapsedTime() >= ASYNC_TASK_TIMEOUT)
+		if (GetElapsedTime() >= AsyncTimeout)
 		{
 			bIsComplete = true;
 			bWasSuccessful = false;
@@ -629,7 +632,10 @@ void FOnlineAsyncTaskSteamLogoffServer::Tick()
 	else
 	{
 		// Fallback timeout in case we don't hear from Steam
-		if (GetElapsedTime() >= ASYNC_TASK_TIMEOUT)
+		float AsyncTimeout = ASYNC_TASK_TIMEOUT;
+		GConfig->GetFloat(TEXT("OnlineSubsystemSteam"), TEXT("OnlineAsyncTaskSteamLogoffServerTimeout"), AsyncTimeout, GEngineIni);
+		
+		if (GetElapsedTime() >= AsyncTimeout)
 		{
 			SessionInt->bSteamworksGameServerConnected = false;
 			SessionInt->GameServerSteamId = NULL;
@@ -1155,7 +1161,10 @@ void FOnlineAsyncTaskSteamFindServerBase::Tick()
 	// Cancel query when we've reached our requested limit
 	bool bReachedSearchLimit = (SearchSettings->SearchResults.Num() >= SearchSettings->MaxSearchResults) ? true : false;
 	// Check for activity timeout
-	bool bTimedOut = (ElapsedTime >= ASYNC_TASK_TIMEOUT) ? true : false;
+	float AsyncTimeout = ASYNC_TASK_TIMEOUT;
+	GConfig->GetFloat(TEXT("OnlineSubsystemSteam"), TEXT("OnlineAsyncTaskSteamFindServerBaseTimeout"), AsyncTimeout, GEngineIni);
+	
+	bool bTimedOut = (ElapsedTime >= AsyncTimeout) ? true : false;
 	// Check for proper completion
 	bool bServerSearchComplete = (bServerRefreshComplete && PendingSearchResults.Num() == 0) ? true : false;
 	if ( bReachedSearchLimit || bTimedOut || bServerSearchComplete)
