GetPlayerNickname() Steam fix - use GetFriendPersonaName with 64-bit ID
Developer note: Code has been reworked, but the target issue is the same

diff --git a/Engine/Plugins/Online/OnlineSubsystemSteam/Source/Private/OnlineIdentityInterfaceSteam.cpp b/Engine/Plugins/Online/OnlineSubsystemSteam/Source/Private/OnlineIdentityInterfaceSteam.cpp
--- a/Engine/Plugins/Online/OnlineSubsystemSteam/Source/Private/OnlineIdentityInterfaceSteam.cpp
+++ b/Engine/Plugins/Online/OnlineSubsystemSteam/Source/Private/OnlineIdentityInterfaceSteam.cpp
@@ -164,10 +164,11 @@ FString FOnlineIdentitySteam::GetPlayerNickname(int32 LocalUserNum) const
 
 FString FOnlineIdentitySteam::GetPlayerNickname(const FUniqueNetId& UserId) const
 {
+	const FUniqueNetIdSteam& SteamId = (const FUniqueNetIdSteam&)UserId;
 	if (SteamFriendsPtr != NULL)
 	{
-		const char* PersonaName = SteamFriendsPtr->GetPersonaName();
-		return FString(UTF8_TO_TCHAR(PersonaName));
+		FString PersonaName = UTF8_TO_TCHAR(SteamFriendsPtr->GetFriendPersonaName(SteamId));
+		return PersonaName;
 	}
 	return FString(TEXT(""));
 }
