Add entitlements and microtransaction implementation support for Steam.

The important addition for Steam is entitlements, Steam apps that can purchase another Steam app through the Steam store (commonly used for expansion packs or other types of DLC) will now have an entry when calling the purchase interface's QueryReceipts functionality.

There is additional support added for Steam's microtransaction API as well. There is little we can actually do in the engine, because all of the API requires calls to and from a trusted backend server. For this reason, the Steam purchasing link was added. It is simply an ease-of-use feature for coders who implement all of their purchasing support via the Online Subsystems. Simply define a server link with calls for initiate transaction and offer querying and all of your existing purchasing code will work just as it does on the other platforms.

diff --git a/Engine/Plugins/Online/OnlineSubsystemSteam/Source/Private/OnlinePurchaseInterfaceSteam.cpp b/Engine/Plugins/Online/OnlineSubsystemSteam/Source/Private/OnlinePurchaseInterfaceSteam.cpp
new file mode 100644
--- /dev/null
+++ b/Engine/Plugins/Online/OnlineSubsystemSteam/Source/Private/OnlinePurchaseInterfaceSteam.cpp
@@ -0,0 +1,205 @@
+// Copyright Epic Games, Inc. All Rights Reserved.
+#include "OnlinePurchaseInterfaceSteam.h"
+#include "OnlineSubsystemSteam.h"
+#include "OnlineError.h"
+#include "Misc/ConfigCacheIni.h"
+#include "Interfaces/OnlineStoreInterfaceV2.h"
+
+#define DYNAMIC_FIELD_ITEM_CATEGORY TEXT("ItemCategory")
+#define DYNAMIC_FIELD_ITEM_NAMESPACE TEXT("ItemNamespace")
+
+FOnlinePurchaseSteam::FOnlinePurchaseSteam(FOnlineSubsystemSteam* InSteamSubsystem)
+	: SteamSubsystem(InSteamSubsystem)
+{
+
+#if USE_STEAM_DEBUG_PURCHASING_LINK
+	TSharedRef<FSteamPurchasingServerLinkDebug> DebugLink = MakeShared<FSteamPurchasingServerLinkDebug>();
+	RegisterServerLink(DebugLink);
+	DebugLink->OffersUpdatedDelegate.ExecuteIfBound(DebugLink->ConfigCategories, DebugLink->ConfigMtxns);
+#endif // USE_STEAM_DEBUG_PURCHASING_LINK
+
+}
+
+FOnlinePurchaseSteam::~FOnlinePurchaseSteam()
+{
+
+}
+
+bool FOnlinePurchaseSteam::IsAllowedToPurchase(const FUniqueNetId& UserId)
+{
+	return true;
+}
+
+void FOnlinePurchaseSteam::Checkout(const FUniqueNetId& UserId, const FPurchaseCheckoutRequest& CheckoutRequest, const FOnPurchaseCheckoutComplete& Delegate)
+{
+	if(!ServerLink.IsValid())
+	{
+		SteamSubsystem->ExecuteNextTick([Delegate]() {
+			UE_LOG_ONLINE_PURCHASE(Error, TEXT("Could not checkout: Steam Server Link is invalid"));
+
+			const TSharedRef<FPurchaseReceipt> PurchaseReceipt = MakeShared<FPurchaseReceipt>();
+			PurchaseReceipt->TransactionState = EPurchaseTransactionState::Failed;
+
+			Delegate.ExecuteIfBound(FOnlineError(EOnlineErrorResult::Unknown), PurchaseReceipt);
+		});
+		return;
+	}
+
+	TArray<FSteamPurchaseDef> Mtxns;
+	for(const FPurchaseCheckoutRequest::FPurchaseOfferEntry& PurchaseOffer : CheckoutRequest.PurchaseOffers)
+	{
+		FSteamPurchaseDef& PurchaseDef = Mtxns.Emplace_GetRef();
+		PurchaseDef.Quantity = PurchaseOffer.Quantity;
+		PurchaseDef.Mtxn->DynamicFields.Add(DYNAMIC_FIELD_ITEM_NAMESPACE, PurchaseOffer.OfferNamespace);
+		PurchaseDef.Mtxn->OfferId = PurchaseOffer.OfferId;
+	}
+
+	ServerLink->InitiateTransaction(UserId, Mtxns, Delegate);
+}
+
+void FOnlinePurchaseSteam::Checkout(const FUniqueNetId& UserId, const FPurchaseCheckoutRequest& CheckoutRequest, const FOnPurchaseReceiptlessCheckoutComplete& Delegate)
+{
+	Checkout(UserId, CheckoutRequest, FOnPurchaseCheckoutComplete::CreateLambda([this, &Delegate](const FOnlineError& Result, const TSharedRef<FPurchaseReceipt>& Receipt)
+	{
+		Delegate.ExecuteIfBound(Result);
+	}));
+}
+
+void FOnlinePurchaseSteam::FinalizePurchase(const FUniqueNetId& UserId, const FString& ReceiptId)
+{
+	ServerLink->FinalizePurchase(UserId, ReceiptId);
+}
+
+void FOnlinePurchaseSteam::RedeemCode(const FUniqueNetId& UserId, const FRedeemCodeRequest& RedeemCodeRequest, const FOnPurchaseRedeemCodeComplete& Delegate)
+{
+	SteamSubsystem->ExecuteNextTick([Delegate]() {
+		Delegate.ExecuteIfBound(FOnlineError(EOnlineErrorResult::Unknown), MakeShared<FPurchaseReceipt>());
+	});
+}
+
+void FOnlinePurchaseSteam::QueryReceipts(const FUniqueNetId& UserId, bool bRestoreReceipts, const FOnQueryReceiptsComplete& Delegate)
+{
+	CachedDLC.Empty();
+	int32 DLCCount = SteamApps()->GetDLCCount();
+	for (int i = 0; i < DLCCount; ++i) {
+		AppId_t OutAppId;
+		bool OutbAvailable;
+		char OutName[128];
+		bool bSuccess = SteamApps()->BGetDLCDataByIndex(i, &OutAppId, &OutbAvailable, OutName, 128);
+		if (bSuccess) {
+			FPurchaseReceipt& DLCRef = CachedDLC.Emplace_GetRef();
+			DLCRef.AddReceiptOffer(FOfferNamespace(), ANSI_TO_TCHAR(OutName), 1);
+		}
+	}
+}
+
+void FOnlinePurchaseSteam::GetReceipts(const FUniqueNetId& UserId, TArray<FPurchaseReceipt>& OutReceipts) const
+{
+	OutReceipts.Empty();
+	for(const FPurchaseReceipt& Receipt : CachedDLC)
+	{
+		OutReceipts.Add(Receipt);
+	}
+}
+
+void FOnlinePurchaseSteam::FinalizeReceiptValidationInfo(const FUniqueNetId& UserId, FString& InReceiptValidationInfo, const FOnFinalizeReceiptValidationInfoComplete& Delegate)
+{
+	ServerLink->FinalizeReceiptValidationInfo(UserId, InReceiptValidationInfo, Delegate);
+}
+
+void FOnlinePurchaseSteam::RegisterServerLink(TSharedRef<ISteamPurchasingServerLink> InServerLink)
+{
+	if(ensure(!ServerLink.IsValid()))
+	{
+		ServerLink = InServerLink;
+		ServerLink->OffersUpdatedDelegate.BindRaw(this, &FOnlinePurchaseSteam::OnServerLinkOffersUpdated);
+	}
+}
+
+void FOnlinePurchaseSteam::OnServerLinkOffersUpdated(TArray<FString> Categories, TArray<FOnlineStoreOfferRef> Mtxns)
+{
+	CachedCategories = Categories;
+	CachedMtxns = Mtxns;
+}
+
+
+void FOnlinePurchaseSteam::GetStoreCategories(TArray<FOnlineStoreCategory>& OutCategories) const
+{
+	OutCategories.Empty();
+	for(int i = 0; i < CachedCategories.Num(); i++)
+	{
+		FOnlineStoreCategory& NewCategory = OutCategories.Emplace_GetRef();
+		NewCategory.Id = FString::Printf(TEXT("%d"), i);
+		NewCategory.Description = FText::FromString(CachedCategories[i]);
+	}
+}
+
+void FOnlinePurchaseSteam::GetStoreOffers(TArray<FOnlineStoreOfferRef>& OutOffers) const
+{
+	OutOffers.Empty();
+	for(const FOnlineStoreOfferRef& CachedMtxn : CachedMtxns)
+	{
+		OutOffers.Add(CachedMtxn);
+	}
+}
+
+TSharedPtr<FOnlineStoreOffer> FOnlinePurchaseSteam::GetStoreOffer(const FUniqueOfferId& OfferId) const
+{
+	for (const FOnlineStoreOfferRef& CachedMtxn : CachedMtxns)
+	{
+		if(CachedMtxn->OfferId.Equals(OfferId))
+		{
+			return CachedMtxn;
+		}
+	}
+	return nullptr;
+}
+
+
+#if USE_STEAM_DEBUG_PURCHASING_LINK
+FSteamPurchasingServerLinkDebug::FSteamPurchasingServerLinkDebug()
+{
+	TArray<FString> ConfigMtxnStrings;
+	if (GConfig->GetArray(TEXT("OnlineSubsystemSteam"), TEXT("StaticMicrotransactions"), ConfigMtxnStrings, GEngineIni))
+	{
+		for(const FString& Mtxn : ConfigMtxnStrings)
+		{
+			FString ItemId, ItemNamespace, ItemDescription, ItemCost, ItemTitle, ItemCategory;
+
+			FParse::Value(*Mtxn, TEXT("Id="), ItemId);
+			FParse::Value(*Mtxn, TEXT("Namespace="), ItemNamespace);
+			FParse::Value(*Mtxn, TEXT("Title="), ItemTitle);
+			FParse::Value(*Mtxn, TEXT("Description="), ItemDescription);
+			FParse::Value(*Mtxn, TEXT("Cost="), ItemCost);
+			FParse::Value(*Mtxn, TEXT("Category="), ItemCategory);
+
+			FOnlineStoreOfferRef& NewMtxn = ConfigMtxns.Add_GetRef(MakeShared<FOnlineStoreOffer>());
+			NewMtxn->Description = FText::FromString(ItemDescription);
+			NewMtxn->Title = FText::FromString(ItemTitle);
+			NewMtxn->OfferId = ItemId;
+			NewMtxn->DynamicFields.Add(DYNAMIC_FIELD_ITEM_NAMESPACE, ItemNamespace);
+			NewMtxn->DynamicFields.Add(DYNAMIC_FIELD_ITEM_CATEGORY, ItemCategory);
+			NewMtxn->RegularPriceText = FText::FromString(ItemCost);
+		}
+	}
+	GConfig->GetArray(TEXT("OnlineSubsystemSteam"), TEXT("StaticMicrotransactionCategories"), ConfigCategories, GEngineIni);
+}
+
+void FSteamPurchasingServerLinkDebug::InitiateTransaction(const FUniqueNetId& UserId, TArray<FSteamPurchaseDef> Mtxns, const FOnPurchaseCheckoutComplete& Delegate)
+{
+	TSharedRef<FPurchaseReceipt> Receipt = MakeShared<FPurchaseReceipt>();
+
+	for(const FSteamPurchaseDef& PurchaseDef : Mtxns)
+	{
+		FString Namespace;
+		if(PurchaseDef.Mtxn->DynamicFields.Contains(DYNAMIC_FIELD_ITEM_NAMESPACE))
+		{
+			Namespace = PurchaseDef.Mtxn->DynamicFields.FindChecked(DYNAMIC_FIELD_ITEM_NAMESPACE);
+		}
+		Receipt->AddReceiptOffer(Namespace, PurchaseDef.Mtxn->OfferId, PurchaseDef.Quantity);
+	}
+
+	Delegate.ExecuteIfBound(FOnlineError::Success(), Receipt);
+}
+
+#endif //USE_STEAM_DEBUG_PURCHASING_LINK
\ No newline at end of file
diff --git a/Engine/Plugins/Online/OnlineSubsystemSteam/Source/Private/OnlineStoreInterfaceSteam.cpp b/Engine/Plugins/Online/OnlineSubsystemSteam/Source/Private/OnlineStoreInterfaceSteam.cpp
new file mode 100644
--- /dev/null
+++ b/Engine/Plugins/Online/OnlineSubsystemSteam/Source/Private/OnlineStoreInterfaceSteam.cpp
@@ -0,0 +1,53 @@
+// Copyright Epic Games, Inc. All Rights Reserved.
+#include "OnlineStoreInterfaceSteam.h"
+#include "OnlineSubsystemSteam.h"
+
+FOnlineStoreSteam::FOnlineStoreSteam(FOnlineSubsystemSteam* InSteamSubsystem)
+	: SteamSubsystem(InSteamSubsystem)
+{
+	SteamPurchaseInt = StaticCastSharedRef<FOnlinePurchaseSteam>(InSteamSubsystem->GetPurchaseInterface().ToSharedRef());
+}
+
+FOnlineStoreSteam::~FOnlineStoreSteam()
+{
+
+}
+
+void FOnlineStoreSteam::QueryCategories(const FUniqueNetId& UserId, const FOnQueryOnlineStoreCategoriesComplete& Delegate)
+{
+	SteamSubsystem->ExecuteNextTick([Delegate]() {
+		Delegate.ExecuteIfBound(true, TEXT(""));
+	});
+}
+
+void FOnlineStoreSteam::GetCategories(TArray<FOnlineStoreCategory>& OutCategories) const
+{
+	return SteamPurchaseInt->GetStoreCategories(OutCategories);
+}
+
+void FOnlineStoreSteam::QueryOffersByFilter(const FUniqueNetId& UserId, const FOnlineStoreFilter& Filter, const FOnQueryOnlineStoreOffersComplete& Delegate)
+{
+	SteamSubsystem->ExecuteNextTick([Delegate]() {
+		TArray<FUniqueOfferId> OfferIds;
+		Delegate.ExecuteIfBound(true, OfferIds, TEXT(""));
+	});
+}
+
+void FOnlineStoreSteam::QueryOffersById(const FUniqueNetId& UserId, const TArray<FUniqueOfferId>& OfferIds, const FOnQueryOnlineStoreOffersComplete& Delegate)
+{
+	SteamSubsystem->ExecuteNextTick([Delegate]() {
+		TArray<FUniqueOfferId> OfferIds;
+		Delegate.ExecuteIfBound(true, OfferIds, TEXT(""));
+	});
+}
+
+void FOnlineStoreSteam::GetOffers(TArray<FOnlineStoreOfferRef>& OutOffers) const
+{
+	return SteamPurchaseInt->GetStoreOffers(OutOffers);
+}
+
+TSharedPtr<FOnlineStoreOffer> FOnlineStoreSteam::GetOffer(const FUniqueOfferId& OfferId) const
+{
+	return SteamPurchaseInt->GetStoreOffer(OfferId);
+}
+
diff --git a/Engine/Plugins/Online/OnlineSubsystemSteam/Source/Private/OnlineStoreInterfaceSteam.h b/Engine/Plugins/Online/OnlineSubsystemSteam/Source/Private/OnlineStoreInterfaceSteam.h
new file mode 100644
--- /dev/null
+++ b/Engine/Plugins/Online/OnlineSubsystemSteam/Source/Private/OnlineStoreInterfaceSteam.h
@@ -0,0 +1,38 @@
+// Copyright Epic Games, Inc. All Rights Reserved.
+
+#pragma once
+
+#include "OnlineSubsystemSteamPrivate.h"
+#include "Interfaces/OnlineStoreInterfaceV2.h"
+#include "OnlinePurchaseInterfaceSteam.h"
+#include "OnlineSubsystemSteamPackage.h"
+
+class FOnlineSubsystemSteam;
+class FStoreTaskBase;
+
+/**
+ * Implementation for online store via Steam
+ */
+class FOnlineStoreSteam :
+	public IOnlineStoreV2,
+	public TSharedFromThis<FOnlineStoreSteam, ESPMode::ThreadSafe>
+{
+public:
+	FOnlineStoreSteam(FOnlineSubsystemSteam* InSteamSubsystem);
+	virtual ~FOnlineStoreSteam();
+
+public:// IOnlineStoreV2
+	virtual void QueryCategories(const FUniqueNetId& UserId, const FOnQueryOnlineStoreCategoriesComplete& Delegate) override;
+	virtual void GetCategories(TArray<FOnlineStoreCategory>& OutCategories) const override;
+	virtual void QueryOffersByFilter(const FUniqueNetId& UserId, const FOnlineStoreFilter& Filter, const FOnQueryOnlineStoreOffersComplete& Delegate) override;
+	virtual void QueryOffersById(const FUniqueNetId& UserId, const TArray<FUniqueOfferId>& OfferIds, const FOnQueryOnlineStoreOffersComplete& Delegate) override;
+	virtual void GetOffers(TArray<FOnlineStoreOfferRef>& OutOffers) const override;
+	virtual TSharedPtr<FOnlineStoreOffer> GetOffer(const FUniqueOfferId& OfferId) const override;
+
+
+PACKAGE_SCOPE:
+	FOnlineSubsystemSteam* SteamSubsystem;
+	TSharedPtr<FOnlinePurchaseSteam> SteamPurchaseInt;
+	TArray<FOnlineStoreCategory> CachedCategories;
+	TMap<FUniqueOfferId, FOnlineStoreOfferRef> CachedOffers;
+};
diff --git a/Engine/Plugins/Online/OnlineSubsystemSteam/Source/Private/OnlineSubsystemSteam.cpp b/Engine/Plugins/Online/OnlineSubsystemSteam/Source/Private/OnlineSubsystemSteam.cpp
--- a/Engine/Plugins/Online/OnlineSubsystemSteam/Source/Private/OnlineSubsystemSteam.cpp
+++ b/Engine/Plugins/Online/OnlineSubsystemSteam/Source/Private/OnlineSubsystemSteam.cpp
@@ -21,6 +21,8 @@
 #include "OnlineAuthInterfaceSteam.h"
 #include "OnlineAuthInterfaceUtilsSteam.h"
 #include "OnlineEncryptedAppTicketInterfaceSteam.h"
+#include "OnlinePurchaseInterfaceSteam.h"
+#include "OnlineStoreInterfaceSteam.h"
 #include "VoiceInterfaceSteam.h"
 
 #include "SteamSharedModule.h"
@@ -317,6 +319,16 @@ IOnlinePresencePtr FOnlineSubsystemSteam::GetPresenceInterface() const
 	return PresenceInterface;
 }
 
+IOnlinePurchasePtr FOnlineSubsystemSteam::GetPurchaseInterface() const
+{
+	return PurchaseInterface;
+}
+
+IOnlineStoreV2Ptr FOnlineSubsystemSteam::GetStoreV2Interface() const
+{
+	return StoreInterface;
+}
+
 IOnlineChatPtr FOnlineSubsystemSteam::GetChatInterface() const
 {
 	return nullptr;
@@ -442,6 +454,8 @@ bool FOnlineSubsystemSteam::Init()
 			ExternalUIInterface = MakeShareable(new FOnlineExternalUISteam(this));
 			AchievementsInterface = MakeShareable(new FOnlineAchievementsSteam(this));
 			EncryptedAppTicketInterface = MakeShareable(new FOnlineEncryptedAppTicketSteam(this));
+			PurchaseInterface = MakeShareable(new FOnlinePurchaseSteam(this));
+			StoreInterface = MakeShareable(new FOnlineStoreSteam(this));
 
 			// Kick off a download/cache of the current user's stats
 			LeaderboardsInterface->CacheCurrentUsersStats();
@@ -510,6 +524,8 @@ bool FOnlineSubsystemSteam::Shutdown()
 	DESTRUCT_INTERFACE(PingInterface);
 	DESTRUCT_INTERFACE(SessionInterface);
 	DESTRUCT_INTERFACE(PresenceInterface);
+	DESTRUCT_INTERFACE(StoreInterface);
+	DESTRUCT_INTERFACE(PurchaseInterface);
 
 #undef DESTRUCT_INTERFACE
 
diff --git a/Engine/Plugins/Online/OnlineSubsystemSteam/Source/Public/OnlinePurchaseInterfaceSteam.h b/Engine/Plugins/Online/OnlineSubsystemSteam/Source/Public/OnlinePurchaseInterfaceSteam.h
new file mode 100644
--- /dev/null
+++ b/Engine/Plugins/Online/OnlineSubsystemSteam/Source/Public/OnlinePurchaseInterfaceSteam.h
@@ -0,0 +1,154 @@
+// Copyright Epic Games, Inc. All Rights Reserved.
+
+#pragma once
+
+#include "OnlineSubsystemSteamPrivate.h"
+#include "Interfaces/OnlinePurchaseInterface.h"
+#include "OnlineSubsystemSteamTypes.h"
+#include "Interfaces/OnlineStoreInterfaceV2.h"
+
+#define USE_STEAM_DEBUG_PURCHASING_LINK 1
+
+class FOnlineSubsystemSteam;
+
+/*
+*  A good amount of Steam purchasing and entitlements cannot be done entirely in-game, it requires the use of a backend server to use the web api for some flows.
+*  This interface serves as a method of preserving the universal platform interface for the Online Subsystems while allowing you to add custom backend code required for full functionality.
+*  Epic does not have an officially licensed solution for the implementation required in this server.
+* 
+*  Note: Microtransactions made using the server link will not show up via QueryReceipts
+*/
+/*struct FSteamMtxnDef
+{
+
+public:
+	FOfferNamespace Namespace;
+	FUniqueOfferId OfferId;
+	FString Description;
+	FString CostString;
+	FString Category;
+};*/
+
+struct FSteamPurchaseDef
+{
+
+public:
+	FOnlineStoreOfferRef Mtxn;
+	int32_t Quantity;
+};
+
+/**
+ * Delegate called when the server wants to update our list of offers.
+ *
+ * @param bWasSuccessful whether the async call completed properly or not
+ */
+DECLARE_MULTICAST_DELEGATE_TwoParams(FOnOffersUpdated, TArray<FString> /*Categories*/, TArray<FOnlineStoreOfferRef> /*Mtxns*/);
+typedef FOnOffersUpdated::FDelegate FOnOffersUpdatedDelegate;
+
+class ISteamPurchasingServerLink
+{
+
+public:
+	virtual ~ISteamPurchasingServerLink() = default;
+
+	/*
+	*	Required virtual function, and the main use case of the server link.
+	*   Use this to call into your server to use the Steam Web API and complete the microtransaction purchase request.
+	*   Your backend server should be using the InitTxn API
+	*/
+	virtual void InitiateTransaction(const FUniqueNetId& UserId, TArray<FSteamPurchaseDef> Mtxns, const FOnPurchaseCheckoutComplete& Delegate) = 0;
+
+	/*
+	*	To use offer enumeration in the purchase interface, you must call these delegates after connecting the server link to refresh the purchase and store interface's offers.
+	*/
+	FOnOffersUpdatedDelegate OffersUpdatedDelegate;
+
+	// Optional Functions
+	
+	/*
+	*  Overrride this function if your game uses the FinalizePurchase functionality to consume a microtransaction and mark it as properly used.
+	*  Your backend server should be calling the FinalizeTxn API, if it did not do so as part of the InitiateTransaction flow
+	*/
+	virtual void FinalizePurchase(const FUniqueNetId& UserId, const FString& ReceiptId) {
+		UE_LOG_ONLINE_PURCHASE(Error, TEXT("FinalizePurchase called but server link does not override!"))
+	}
+
+	/*
+	*  Overrride this function if your game uses the FinalizeReceiptValidationInfo. 
+	*  The Steam API does not directly provide anything useful for this call (other platforms would provide something like a JWT token). It is generally recommended to simply make that connection on the backend using the web API
+	*/
+	virtual void FinalizeReceiptValidationInfo(const FUniqueNetId& UserId, FString& InReceiptValidationInfo, const FOnFinalizeReceiptValidationInfoComplete& Delegate) {
+		UE_LOG_ONLINE_PURCHASE(Error, TEXT("FinalizeReceiptValidationInfo called but server link does not override!"))
+	}
+};
+
+/* 
+* The Steam purchase interface is designed to handle microtransactions in-game, not DLC entitlements (use the external UI to show the Steam store for those)
+* Steam microtransactions are not defined by the server, but rather by the client. UE lets us define these microtransactions in two ways
+* 
+*  1) Config Driven
+*		Config driven microtransactions can be datamined, and can possibily be tampered if your backend server does not have the proper protections in place to verify the transaction. Therefore it is only recommended for development/debug flows.
+*		To add a new static microtransaction to Steam, add the following in your game's Engine file
+*			
+*				[OnlineSubsystemSteam]
+*					+StaticMicrotransactions=(Id=32bitinteger, Amount=costofitem, Description=myMicroTransaction) 
+* 
+*  2) Server-driven
+*		- Make an implementation of ISteamPurchasingServerLink and implement all functions
+*/
+
+class FOnlinePurchaseSteam
+	: public IOnlinePurchase
+	, public TSharedFromThis<FOnlinePurchaseSteam, ESPMode::ThreadSafe>
+{
+public:
+	FOnlinePurchaseSteam(FOnlineSubsystemSteam* InSteamSubsystem);
+	virtual ~FOnlinePurchaseSteam();
+
+public:
+	//~ Begin IOnlinePurchase Interface
+	virtual bool IsAllowedToPurchase(const FUniqueNetId& UserId) override;
+	virtual void Checkout(const FUniqueNetId& UserId, const FPurchaseCheckoutRequest& CheckoutRequest, const FOnPurchaseCheckoutComplete& Delegate) override;
+	virtual void Checkout(const FUniqueNetId& UserId, const FPurchaseCheckoutRequest& CheckoutRequest, const FOnPurchaseReceiptlessCheckoutComplete& Delegate) override;
+	virtual void FinalizePurchase(const FUniqueNetId& UserId, const FString& ReceiptId) override;
+	virtual void RedeemCode(const FUniqueNetId& UserId, const FRedeemCodeRequest& RedeemCodeRequest, const FOnPurchaseRedeemCodeComplete& Delegate) override;
+	virtual void QueryReceipts(const FUniqueNetId& UserId, bool bRestoreReceipts, const FOnQueryReceiptsComplete& Delegate) override;
+	virtual void GetReceipts(const FUniqueNetId& UserId, TArray<FPurchaseReceipt>& OutReceipts) const override;
+	virtual void FinalizeReceiptValidationInfo(const FUniqueNetId& UserId, FString& InReceiptValidationInfo, const FOnFinalizeReceiptValidationInfoComplete& Delegate) override;
+	//~ End IOnlinePurchase Interface
+
+	void RegisterServerLink(TSharedRef<ISteamPurchasingServerLink> InServerLink);
+
+	void GetStoreCategories(TArray<FOnlineStoreCategory>& OutCategories) const;
+	void GetStoreOffers(TArray<FOnlineStoreOfferRef>& OutOffers) const;
+	TSharedPtr<FOnlineStoreOffer> GetStoreOffer(const FUniqueOfferId& OfferId) const;
+
+private:
+	FOnlineSubsystemSteam* SteamSubsystem;
+	TSharedPtr<ISteamPurchasingServerLink> ServerLink;
+
+	TArray<FOnlineStoreOfferRef> CachedMtxns;
+	TArray<FString> CachedCategories;
+	TArray<FPurchaseReceipt> CachedDLC;
+
+	void OnServerLinkOffersUpdated(TArray<FString> Categories, TArray<FOnlineStoreOfferRef> Mtxns);
+
+};
+
+
+#if USE_STEAM_DEBUG_PURCHASING_LINK
+
+class FSteamPurchasingServerLinkDebug
+	: public ISteamPurchasingServerLink
+{
+public:
+	FSteamPurchasingServerLinkDebug();
+	virtual ~FSteamPurchasingServerLinkDebug() = default;
+
+	virtual void InitiateTransaction(const FUniqueNetId& UserId, TArray<FSteamPurchaseDef> Mtxns, const FOnPurchaseCheckoutComplete& Delegate) override;
+
+	TArray<FOnlineStoreOfferRef> ConfigMtxns;
+	TArray<FString> ConfigCategories;
+};
+
+#endif //USE_STEAM_DEBUG_PURCHASING_LINK
\ No newline at end of file
diff --git a/Engine/Plugins/Online/OnlineSubsystemSteam/Source/Public/OnlineSubsystemSteam.h b/Engine/Plugins/Online/OnlineSubsystemSteam/Source/Public/OnlineSubsystemSteam.h
--- a/Engine/Plugins/Online/OnlineSubsystemSteam/Source/Public/OnlineSubsystemSteam.h
+++ b/Engine/Plugins/Online/OnlineSubsystemSteam/Source/Public/OnlineSubsystemSteam.h
@@ -21,6 +21,8 @@ class FOnlineAuthSteam;
 class FOnlineAuthUtilsSteam;
 class FOnlinePingInterfaceSteam;
 class FOnlineEncryptedAppTicketSteam;
+class FOnlinePurchaseSteam;
+class FOnlineStoreSteam;
 
 /** Forward declarations of all interface classes */
 typedef TSharedPtr<class FOnlineSessionSteam, ESPMode::ThreadSafe> FOnlineSessionSteamPtr;
@@ -37,6 +39,8 @@ typedef TSharedPtr<class FOnlineAuthSteam, ESPMode::ThreadSafe> FOnlineAuthSteam
 typedef TSharedPtr<class FOnlineAuthUtilsSteam, ESPMode::ThreadSafe> FOnlineAuthSteamUtilsPtr;
 typedef TSharedPtr<class FOnlinePingInterfaceSteam, ESPMode::ThreadSafe> FOnlinePingSteamPtr;
 typedef TSharedPtr<class FOnlineEncryptedAppTicketSteam, ESPMode::ThreadSafe> FOnlineEncryptedAppTicketSteamPtr;
+typedef TSharedPtr<class FOnlinePurchaseSteam, ESPMode::ThreadSafe> FOnlinePurchaseSteamPtr;
+typedef TSharedPtr<class FOnlineStoreSteam, ESPMode::ThreadSafe> FOnlineStoreSteamPtr;
 
 /**
 * Delegate fired when a Steam Game Server has completed its login tasks with the Steam backend.
@@ -119,6 +123,12 @@ protected:
 	/** Interface for Steam encrypted application tickets. */
 	FOnlineEncryptedAppTicketSteamPtr EncryptedAppTicketInterface;
 
+	/** Interface for the Purchase interface */
+	FOnlinePurchaseSteamPtr PurchaseInterface;
+
+	/** Interface for the Store interface */
+	FOnlineStoreSteamPtr StoreInterface;
+
 	/** Online async task runnable */
 	class FOnlineAsyncTaskManagerSteam* OnlineAsyncTaskThreadRunnable;
 
@@ -157,6 +167,8 @@ PACKAGE_SCOPE:
 		AuthInterfaceUtils(nullptr),
 		PingInterface(nullptr),
 		EncryptedAppTicketInterface(nullptr),
+		PurchaseInterface(nullptr),
+		StoreInterface(nullptr),
 		OnlineAsyncTaskThreadRunnable(nullptr),
 		OnlineAsyncTaskThread(nullptr),
 		SteamAPIClientHandle(nullptr),
@@ -261,8 +273,8 @@ public:
 	virtual IOnlineIdentityPtr GetIdentityInterface() const override;
 	virtual IOnlineTitleFilePtr GetTitleFileInterface() const override;
 	virtual IOnlineEntitlementsPtr GetEntitlementsInterface() const override;
-	virtual IOnlineStoreV2Ptr GetStoreV2Interface() const override { return nullptr; }
-	virtual IOnlinePurchasePtr GetPurchaseInterface() const override { return nullptr; }
+	virtual IOnlineStoreV2Ptr GetStoreV2Interface() const override;
+	virtual IOnlinePurchasePtr GetPurchaseInterface() const override;
 	virtual IOnlineEventsPtr GetEventsInterface() const override;
 	virtual IOnlineAchievementsPtr GetAchievementsInterface() const override;
 	virtual IOnlineSharingPtr GetSharingInterface() const override;
