Moving code to create/delete steam_appid.txt to SteamShared.

diff --git a/Engine/Plugins/Online/OnlineSubsystemSteam/Source/Private/OnlineSubsystemSteam.cpp b/Engine/Plugins/Online/OnlineSubsystemSteam/Source/Private/OnlineSubsystemSteam.cpp
--- a/Engine/Plugins/Online/OnlineSubsystemSteam/Source/Private/OnlineSubsystemSteam.cpp
+++ b/Engine/Plugins/Online/OnlineSubsystemSteam/Source/Private/OnlineSubsystemSteam.cpp
@@ -2,7 +2,6 @@
 
 #include "OnlineSubsystemSteam.h"
 #include "HAL/FileManager.h"
-#include "HAL/PlatformFile.h"
 #include "HAL/RunnableThread.h"
 #include "Interfaces/IPluginManager.h"
 #include "Misc/StringBuilder.h"
@@ -95,80 +94,23 @@ static void __cdecl SteamworksWarningMessageHookNoOp(int Severity, const char *M
 	// no-op.
 }
 
-inline FString GetSteamAppIdFilename()
-{
-	return FString::Printf(TEXT("%s%s"), FPlatformProcess::BaseDir(), STEAMAPPIDFILENAME);
-}
-
-#if !UE_BUILD_SHIPPING && !UE_BUILD_SHIPPING_WITH_EDITOR
-/** 
- * Write out the steam app id to the steam_appid.txt file before initializing the API
- * @param SteamAppId id assigned to the application by Steam
- */
-static bool WriteSteamAppIdToDisk(int32 SteamAppId)
-{
-	if (SteamAppId > 0)
-	{
-		// Access the physical file writer directly so that we still write next to the executable in CotF builds.
-		FString SteamAppIdFilename = GetSteamAppIdFilename();
-		TUniquePtr<IFileHandle> Handle(IPlatformFile::GetPlatformPhysical().OpenWrite(*SteamAppIdFilename, false, false));
-		if (!Handle)
-		{
-			UE_LOG_ONLINE(Error, TEXT("Failed to create file: %s"), *SteamAppIdFilename);
-			return false;
-		}
-		else
-		{
-			TAnsiStringBuilder<16> AppId(InPlace, SteamAppId);
-			Handle->Write((const uint8*)AppId.GetData(), AppId.Len());
-			return true;
-		}
-	}
-
-	UE_LOG_ONLINE(Warning, TEXT("Steam App Id provided (%d) is invalid, must be greater than 0!"), SteamAppId);
-	return false;
-}
-
-/**
- * Deletes the app id file from disk
- */
-static void DeleteSteamAppIdFromDisk()
-{
-	const FString SteamAppIdFilename = GetSteamAppIdFilename();
-	IPlatformFile& Physical = IPlatformFile::GetPlatformPhysical();
-	if (Physical.FileExists(*SteamAppIdFilename))
-	{
-		Physical.DeleteFile(*SteamAppIdFilename);
-	}
-}
-
-#endif // !UE_BUILD_SHIPPING && !UE_BUILD_SHIPPING_WITH_EDITOR
-
 /**
- * Configure various dev options before initializing Steam
+ * Get relaunch settings from OnlineSubsystemSteam configuration
  *
  * @param RequireRelaunch enforce the Steam client running precondition
  * @param RelaunchAppId appid to launch when the Steam client is loaded
  *
  * @return if this sequence completed without any serious errors
  */
-bool ConfigureSteamInitDevOptions(bool& RequireRelaunch, int32& RelaunchAppId)
+bool GetRelaunchSettings(bool& RequireRelaunch, int32& RelaunchAppId)
 {
 #if !UE_BUILD_SHIPPING && !UE_BUILD_SHIPPING_WITH_EDITOR
-	// Write out the steam_appid.txt file before launching
+	// Get SteamDevAppId and set it as RelaunchAppId
 	if (!GConfig->GetInt(TEXT("OnlineSubsystemSteam"), TEXT("SteamDevAppId"), RelaunchAppId, GEngineIni))
 	{
 		UE_LOG_ONLINE(Warning, TEXT("Missing SteamDevAppId key in OnlineSubsystemSteam of DefaultEngine.ini"));
 		return false;
 	}
-	else
-	{
-		if (!WriteSteamAppIdToDisk(RelaunchAppId))
-		{
-			UE_LOG_ONLINE(Warning, TEXT("Could not create/update the steam_appid.txt file! Make sure the directory is writable and there isn't another instance using this file"));
-			return false;
-		}
-	}
 
 	// Should the game force a relaunch in Steam if the client isn't already loaded
 	if (!GConfig->GetBool(TEXT("OnlineSubsystemSteam"), TEXT("bRelaunchInSteam"), RequireRelaunch, GEngineIni))
@@ -393,7 +335,7 @@ bool FOnlineSubsystemSteam::Init()
 	bool bRelaunchInSteam = false;
 	int RelaunchAppId = 0;
 
-	if (!ConfigureSteamInitDevOptions(bRelaunchInSteam, RelaunchAppId))
+	if (!GetRelaunchSettings(bRelaunchInSteam, RelaunchAppId))
 	{
 		UE_LOG_ONLINE(Warning, TEXT("Could not set up the steam environment! Falling back to another OSS."));
 		return false;
@@ -533,10 +475,6 @@ bool FOnlineSubsystemSteam::Shutdown()
 
 	ShutdownSteamworks();
 
-#if !UE_BUILD_SHIPPING && !UE_BUILD_SHIPPING_WITH_EDITOR
-	DeleteSteamAppIdFromDisk();
-#endif // !UE_BUILD_SHIPPING && !UE_BUILD_SHIPPING_WITH_EDITOR
-
 	return true;
 }
 
diff --git a/Engine/Plugins/Online/OnlineSubsystemSteam/Source/Private/OnlineSubsystemSteamPrivate.h b/Engine/Plugins/Online/OnlineSubsystemSteam/Source/Private/OnlineSubsystemSteamPrivate.h
--- a/Engine/Plugins/Online/OnlineSubsystemSteam/Source/Private/OnlineSubsystemSteamPrivate.h
+++ b/Engine/Plugins/Online/OnlineSubsystemSteam/Source/Private/OnlineSubsystemSteamPrivate.h
@@ -9,8 +9,6 @@
 
 /** URL Prefix when using Steam socket connection */
 #define STEAM_URL_PREFIX TEXT("steam.")
-/** Filename containing the appid during development */
-#define STEAMAPPIDFILENAME TEXT("steam_appid.txt")
 
 /** pre-pended to all steam logging */
 #undef ONLINE_LOG_PREFIX
diff --git a/Engine/Plugins/Runtime/Steam/SteamShared/Source/SteamShared/Private/SteamSharedModule.cpp b/Engine/Plugins/Runtime/Steam/SteamShared/Source/SteamShared/Private/SteamSharedModule.cpp
--- a/Engine/Plugins/Runtime/Steam/SteamShared/Source/SteamShared/Private/SteamSharedModule.cpp
+++ b/Engine/Plugins/Runtime/Steam/SteamShared/Source/SteamShared/Private/SteamSharedModule.cpp
@@ -12,6 +12,10 @@
 #error Steam SDK not located! Expected to be found in Engine/Source/ThirdParty/Steamworks/{SteamVersion}
 #endif // STEAM_SDK_INSTALLED
 
+#if !UE_BUILD_SHIPPING && !UE_BUILD_SHIPPING_WITH_EDITOR
+#include "HAL/PlatformFile.h"
+#endif // !UE_BUILD_SHIPPING && !UE_BUILD_SHIPPING_WITH_EDITOR
+
 // Steam API for Initialization
 THIRD_PARTY_INCLUDES_START
 #include "steam/steam_api.h"
@@ -62,6 +66,65 @@ bool FSteamSharedModule::CanLoadClientDllsOnServer() const
 	return false;
 }
 
+#if !UE_BUILD_SHIPPING && !UE_BUILD_SHIPPING_WITH_EDITOR
+/** Filename containing the appid during development */
+#define STEAMAPPIDFILENAME TEXT("steam_appid.txt")
+
+inline FString GetSteamAppIdFilename()
+{
+	return FString::Printf(TEXT("%s%s"), FPlatformProcess::BaseDir(), STEAMAPPIDFILENAME);
+}
+
+/**
+ * Write the steam app id to the steam_appid.txt file before initializing the API
+ */
+static bool WriteSteamAppIdToDisk()
+{
+	// Write out the steam_appid.txt file before launching
+	int32 SteamDevAppId;
+
+	if (!GConfig->GetInt(TEXT("OnlineSubsystemSteam"), TEXT("SteamDevAppId"), SteamDevAppId, GEngineIni))
+	{
+		UE_LOG(LogSteamShared, Warning, TEXT("Missing SteamDevAppId key in OnlineSubsystemSteam configuration in .ini file."));
+		return false;
+	}
+
+	if (SteamDevAppId > 0)
+	{
+		// Access the physical file writer directly so that we still write next to the executable in CotF builds.
+		FString SteamAppIdFilename = GetSteamAppIdFilename();
+		TUniquePtr<IFileHandle> Handle(IPlatformFile::GetPlatformPhysical().OpenWrite(*SteamAppIdFilename, false, false));
+		if (!Handle)
+		{
+			UE_LOG(LogSteamShared, Warning, TEXT("Failed to create file: %s"), *SteamAppIdFilename);
+			return false;
+		}
+		else
+		{
+			TAnsiStringBuilder<16> AppId(InPlace, SteamDevAppId);
+			Handle->Write((const uint8*)AppId.GetData(), AppId.Len());
+			return true;
+		}
+	}
+
+	UE_LOG(LogSteamShared, Warning, TEXT("Steam App Id provided (%d) is invalid, must be greater than 0!"), SteamDevAppId);
+	return false;
+}
+
+/**
+ * Deletes the steam app id file from disk
+ */
+static void DeleteSteamAppIdFromDisk()
+{
+	const FString SteamAppIdFilename = GetSteamAppIdFilename();
+	IPlatformFile& Physical = IPlatformFile::GetPlatformPhysical();
+	if (Physical.FileExists(*SteamAppIdFilename))
+	{
+		Physical.DeleteFile(*SteamAppIdFilename);
+	}
+}
+#endif // !UE_BUILD_SHIPPING && !UE_BUILD_SHIPPING_WITH_EDITOR
+
 void FSteamSharedModule::StartupModule()
 {
 	// On startup load the modules. Anyone who uses this shared library can guarantee
@@ -96,6 +159,10 @@ void FSteamSharedModule::ShutdownModule()
 
 	// Here we are no longer loaded, so we need to override any DLL handles still open and unlink the DLLs.
 	UnloadSteamModules();
+
+#if !UE_BUILD_SHIPPING && !UE_BUILD_SHIPPING_WITH_EDITOR
+	DeleteSteamAppIdFromDisk();
+#endif // !UE_BUILD_SHIPPING && !UE_BUILD_SHIPPING_WITH_EDITOR
 }
 
 TSharedPtr<class FSteamClientInstanceHandler> FSteamSharedModule::ObtainSteamClientInstanceHandle()
@@ -173,6 +240,13 @@ void FSteamSharedModule::LoadSteamModules()
 	{
 		return;
 	}
+
+#if !UE_BUILD_SHIPPING && !UE_BUILD_SHIPPING_WITH_EDITOR
+	if (!WriteSteamAppIdToDisk())
+	{
+		UE_LOG(LogSteamShared, Warning, TEXT("Failed to create steam_appid.txt. If the file doesn't exist, SteamAPI_Init will fail."));
+	}
+#endif // !UE_BUILD_SHIPPING && !UE_BUILD_SHIPPING_WITH_EDITOR 
 		
 	UE_LOG(LogSteamShared, Display, TEXT("Loading Steam SDK %s"), STEAM_SDK_VER);
 
