Addition of missing support for lobby usage variable in Steam sessions

diff --git a/Engine/Plugins/Online/OnlineSubsystem/Source/Private/OnlineSessionSettings.cpp b/Engine/Plugins/Online/OnlineSubsystem/Source/Private/OnlineSessionSettings.cpp
--- a/Engine/Plugins/Online/OnlineSubsystem/Source/Private/OnlineSessionSettings.cpp
+++ b/Engine/Plugins/Online/OnlineSubsystem/Source/Private/OnlineSessionSettings.cpp
@@ -59,6 +59,7 @@ void DumpSessionSettings(const FOnlineSessionSettings* SessionSettings)
 		UE_LOG_ONLINE_SESSION(Verbose, TEXT("\tbAllowJoinInProgress: %s"), SessionSettings->bAllowJoinInProgress ? TEXT("true") : TEXT("false"));
 		UE_LOG_ONLINE_SESSION(Verbose, TEXT("\tbAllowInvites: %s"), SessionSettings->bAllowInvites ? TEXT("true") : TEXT("false"));
 		UE_LOG_ONLINE_SESSION(Verbose, TEXT("\tbUsesPresence: %s"), SessionSettings->bUsesPresence ? TEXT("true") : TEXT("false"));
+		UE_LOG_ONLINE_SESSION(Verbose, TEXT("\tbUseLobbiesIfAvailable: %s"), SessionSettings->bUseLobbiesIfAvailable ? TEXT("true") : TEXT("false"));
 		UE_LOG_ONLINE_SESSION(Verbose, TEXT("\tbAllowJoinViaPresence: %s"), SessionSettings->bAllowJoinViaPresence ? TEXT("true") : TEXT("false"));
 		UE_LOG_ONLINE_SESSION(Verbose, TEXT("\tbAllowJoinViaPresenceFriendsOnly: %s"), SessionSettings->bAllowJoinViaPresenceFriendsOnly ? TEXT("true") : TEXT("false"));
 		UE_LOG_ONLINE_SESSION(Verbose, TEXT("\tBuildUniqueId: 0x%08x"), SessionSettings->BuildUniqueId);
diff --git a/Engine/Plugins/Online/OnlineSubsystemSteam/Source/Private/OnlineSessionAsyncLobbySteam.cpp b/Engine/Plugins/Online/OnlineSubsystemSteam/Source/Private/OnlineSessionAsyncLobbySteam.cpp
--- a/Engine/Plugins/Online/OnlineSubsystemSteam/Source/Private/OnlineSessionAsyncLobbySteam.cpp
+++ b/Engine/Plugins/Online/OnlineSubsystemSteam/Source/Private/OnlineSessionAsyncLobbySteam.cpp
@@ -111,6 +111,7 @@ static void GetLobbyKeyValuePairsFromSessionSettings(const FOnlineSessionSetting
 	SessionFlags |= (SessionSettings.bUsesStats ? 1 : 0) << BitShift++;
 	SessionFlags |= (SessionSettings.bAllowInvites ? 1 : 0) << BitShift++;
 	SessionFlags |= (SessionSettings.bUsesPresence ? 1 : 0) << BitShift++;
+	SessionFlags |= (SessionSettings.bUseLobbiesIfAvailable ? 1 : 0) << BitShift++;
 	SessionFlags |= (SessionSettings.bAllowJoinViaPresence ? 1 : 0) << BitShift++;
 	SessionFlags |= (SessionSettings.bAllowJoinViaPresenceFriendsOnly ? 1 : 0) << BitShift++;
 	SessionFlags |= (SessionSettings.bAntiCheatProtected ? 1 : 0) << BitShift++;
@@ -279,6 +280,7 @@ bool FillSessionFromLobbyData(FOnlineSubsystemSteam* SteamSubsystem, const FUniq
 			Session.SessionSettings.bUsesStats = (SessionFlags & (1 << BitShift++)) ? true : false;
 			Session.SessionSettings.bAllowInvites = (SessionFlags & (1 << BitShift++)) ? true : false;
 			Session.SessionSettings.bUsesPresence = (SessionFlags & (1 << BitShift++)) ? true : false;
+			Session.SessionSettings.bUseLobbiesIfAvailable = (SessionFlags & (1 << BitShift++)) ? true : false;
 			Session.SessionSettings.bAllowJoinViaPresence = (SessionFlags & (1 << BitShift++)) ? true : false;
 			Session.SessionSettings.bAllowJoinViaPresenceFriendsOnly = (SessionFlags & (1 << BitShift++)) ? true : false;
 			Session.SessionSettings.bAntiCheatProtected = (SessionFlags & (1 << BitShift++)) ? true : false;
