Fix presence crash issue with shared ptr allocation on Steam

diff --git a/Engine/Plugins/Online/OnlineSubsystemSteam/Source/Private/OnlinePresenceInterfaceSteam.cpp b/Engine/Plugins/Online/OnlineSubsystemSteam/Source/Private/OnlinePresenceInterfaceSteam.cpp
--- a/Engine/Plugins/Online/OnlineSubsystemSteam/Source/Private/OnlinePresenceInterfaceSteam.cpp
+++ b/Engine/Plugins/Online/OnlineSubsystemSteam/Source/Private/OnlinePresenceInterfaceSteam.cpp
@@ -314,7 +314,7 @@ EOnlineCachedResult::Type FOnlinePresenceSteam::GetCachedPresence(const FUniqueN
 		return EOnlineCachedResult::NotFound;
 	}
 
-	OutPresence = MakeShared<FOnlineUserPresence>(Found->Get());
+	OutPresence = Found->ToSharedPtr();
 	return EOnlineCachedResult::Success;
 }
 
