- Fix GetLinkedAccountAuthToken in OSSSteam binding to an old style completion delegate, but not unbinding, causing the callback to fire 1 extra time for each previous call.

diff --git a/Engine/Plugins/Online/OnlineSubsystemSteam/Source/Private/OnlineIdentityInterfaceSteam.cpp b/Engine/Plugins/Online/OnlineSubsystemSteam/Source/Private/OnlineIdentityInterfaceSteam.cpp
--- a/Engine/Plugins/Online/OnlineSubsystemSteam/Source/Private/OnlineIdentityInterfaceSteam.cpp
+++ b/Engine/Plugins/Online/OnlineSubsystemSteam/Source/Private/OnlineIdentityInterfaceSteam.cpp
@@ -262,7 +262,8 @@ void FOnlineIdentitySteam::GetLinkedAccountAuthToken(int32 LocalUserNum, const F
 
 	if (TokenType == AppTokenType)
 	{
-		SteamSubsystem->GetEncryptedAppTicketInterface()->OnEncryptedAppTicketResultDelegate.AddLambda([this, LocalUserNum, OnComplete = FOnGetLinkedAccountAuthTokenCompleteDelegate(Delegate)](bool bEncryptedDataAvailable, int32 ResultCode)
+		TSharedRef<FDelegateHandle> DelegateHandle = MakeShared<FDelegateHandle>();
+		*DelegateHandle = SteamSubsystem->GetEncryptedAppTicketInterface()->OnEncryptedAppTicketResultDelegate.AddLambda([this, LocalUserNum, Delegate, DelegateHandle](bool bEncryptedDataAvailable, int32 ResultCode)
 		{
 			FExternalAuthToken ExternalToken;
 			if (bEncryptedDataAvailable)
@@ -270,7 +271,8 @@ void FOnlineIdentitySteam::GetLinkedAccountAuthToken(int32 LocalUserNum, const F
 				SteamSubsystem->GetEncryptedAppTicketInterface()->GetEncryptedAppTicket(ExternalToken.TokenData);
 			}
 			// Pass the info back to the original caller
-			OnComplete.ExecuteIfBound(LocalUserNum, ExternalToken.HasTokenData(), ExternalToken);
+			Delegate.ExecuteIfBound(LocalUserNum, ExternalToken.HasTokenData(), ExternalToken);
+			SteamSubsystem->GetEncryptedAppTicketInterface()->OnEncryptedAppTicketResultDelegate.Remove(*DelegateHandle);
 		});
 		SteamSubsystem->GetEncryptedAppTicketInterface()->RequestEncryptedAppTicket(nullptr, 0);
 	}
