UE: OnlineSubsystemSteam - Steam Auth Token Delegate call Fix for when there's no internet connection

diff --git a/Engine/Plugins/Online/OnlineSubsystemSteam/Source/Private/OnlineIdentityInterfaceSteam.cpp b/Engine/Plugins/Online/OnlineSubsystemSteam/Source/Private/OnlineIdentityInterfaceSteam.cpp
--- a/Engine/Plugins/Online/OnlineSubsystemSteam/Source/Private/OnlineIdentityInterfaceSteam.cpp
+++ b/Engine/Plugins/Online/OnlineSubsystemSteam/Source/Private/OnlineIdentityInterfaceSteam.cpp
@@ -308,6 +308,20 @@ void FOnlineIdentitySteam::GetLinkedAccountAuthToken(int32 LocalUserNum, const F
 					OnComplete.ExecuteIfBound(LocalUserNum, AuthToken.HasTokenString(), AuthToken);
 				}));
 		}
+		else
+		{
+			if (SteamUserPtr == NULL)
+			{
+				UE_LOG_ONLINE_IDENTITY(Warning, TEXT("FOnlineIdentitySteam::GetLinkedAccountAuthToken failed because SteamUserPtr is NULL"));
+			}
+			else
+			{
+				UE_LOG_ONLINE_IDENTITY(Warning, TEXT("FOnlineIdentitySteam::GetLinkedAccountAuthToken failed because User is not logged into Steam"));
+			}
+
+			// fire the delegate with failure for cases like UE-330775
+			Delegate.ExecuteIfBound(LocalUserNum, false, FExternalAuthToken());
+		}
 	}
 	else
 	{	
