Disable Steam on Android

#Synced-CL 49977598

diff --git a/Engine/Plugins/Online/OnlineSubsystemSteam/OnlineSubsystemSteam.uplugin b/Engine/Plugins/Online/OnlineSubsystemSteam/OnlineSubsystemSteam.uplugin
--- a/Engine/Plugins/Online/OnlineSubsystemSteam/OnlineSubsystemSteam.uplugin
+++ b/Engine/Plugins/Online/OnlineSubsystemSteam/OnlineSubsystemSteam.uplugin
@@ -20,7 +20,6 @@
                 "Win64",
                 "Mac",
                 "Linux",
-		"Android"
 			],
 			"PlatformArchitectureDenyList" :
 			[
diff --git a/Engine/Plugins/Runtime/Steam/SteamShared/SteamShared.uplugin b/Engine/Plugins/Runtime/Steam/SteamShared/SteamShared.uplugin
--- a/Engine/Plugins/Runtime/Steam/SteamShared/SteamShared.uplugin
+++ b/Engine/Plugins/Runtime/Steam/SteamShared/SteamShared.uplugin
@@ -17,8 +17,7 @@
 			"PlatformAllowList": [
 				"Win64",
 				"Mac",
-				"Linux",
-				"Android"
+				"Linux"
 			],
 			"PlatformArchitectureDenyList": [
 				"Win64:arm64"
